Added visual editing of analytics configuration settings
Restructured the analytics plugins so that editor only and platform specific functionality are all in one plugin. This requires your project file to be updated

#Ocean
Updated the project due to the plugin changes

diff --git a/Engine/Plugins/Runtime/Analytics/AnalyticsMulticast/AnalyticsMulticast.uplugin b/Engine/Plugins/Runtime/Analytics/AnalyticsMulticast/AnalyticsMulticast.uplugin
--- a/Engine/Plugins/Runtime/Analytics/AnalyticsMulticast/AnalyticsMulticast.uplugin
+++ b/Engine/Plugins/Runtime/Analytics/AnalyticsMulticast/AnalyticsMulticast.uplugin
@@ -2,16 +2,21 @@
     "FileVersion": 3,
 
     "FriendlyName": "Multicast Analytics Provider",
-    "Version": 1,
-    "VersionName": "1.0",
+    "Version": 2,
+    "VersionName": "2.0",
     "CreatedBy": "Epic Games, Inc.",
     "CreatedByURL": "http://epicgames.com",
     "DocsURL": "https://docs.unrealengine.com/latest/INT/Gameplay/Analytics/index.html",
-    "EngineVersion": "4.8.0",
+    "EngineVersion": "4.9.0",
     "Description": "Forwards analytics API calls to a list of analytics providers to log data to multiple services at once",
     "Category": "Analytics",
 
-    "Modules": [
+    "Modules": 
+    [
+        {
+            "Name": "AnalyticsMulticastEditor",
+            "Type": "Editor"
+        },
         {
             "Name": "AnalyticsMulticast",
             "Type": "Runtime",
diff --git a/Engine/Plugins/Runtime/Analytics/AnalyticsMulticast/Source/AnaltyicsMulticastEditor/AnalyticsMulticastEditor.Build.cs b/Engine/Plugins/Runtime/Analytics/AnalyticsMulticast/Source/AnaltyicsMulticastEditor/AnalyticsMulticastEditor.Build.cs
new file mode 100644
--- /dev/null
+++ b/Engine/Plugins/Runtime/Analytics/AnalyticsMulticast/Source/AnaltyicsMulticastEditor/AnalyticsMulticastEditor.Build.cs
@@ -0,0 +1,44 @@
+// Copyright 1998-2015 Epic Games, Inc. All Rights Reserved.
+
+using UnrealBuildTool;
+
+public class AnalyticsMulticastEditor : ModuleRules
+{
+    public AnalyticsMulticastEditor(TargetInfo Target)
+	{
+        PrivateIncludePaths.Add("AnalyticsMulticastEditor/Private");
+
+		PrivateDependencyModuleNames.AddRange(
+			new string[] {
+				"Core",
+				"CoreUObject",
+                "Analytics",
+                "AnalyticsVisualEditing",
+				"Slate",
+				"SlateCore",
+				"Engine",
+				"UnrealEd", // for FAssetEditorManager
+				"PropertyEditor",
+				"WorkspaceMenuStructure",
+				"EditorStyle",
+				"EditorWidgets",
+				"Projects"
+			}
+			);
+
+		PrivateIncludePathModuleNames.AddRange(
+			new string[] {
+				"Settings",
+				"IntroTutorials",
+                "AssetTools"
+			}
+		);
+
+        DynamicallyLoadedModuleNames.AddRange(
+            new string[] {
+                "AssetTools"
+            }
+        );
+
+	}
+}
diff --git a/Engine/Plugins/Runtime/Analytics/AnalyticsMulticast/Source/AnaltyicsMulticastEditor/Classes/AnalyticsMulticastSettings.h b/Engine/Plugins/Runtime/Analytics/AnalyticsMulticast/Source/AnaltyicsMulticastEditor/Classes/AnalyticsMulticastSettings.h
new file mode 100644
--- /dev/null
+++ b/Engine/Plugins/Runtime/Analytics/AnalyticsMulticast/Source/AnaltyicsMulticastEditor/Classes/AnalyticsMulticastSettings.h
@@ -0,0 +1,46 @@
+// Copyright 1998-2015 Epic Games, Inc. All Rights Reserved.
+
+#pragma once
+
+#include "AnalyticsSettings.h"
+#include "AnalyticsMulticastSettings.generated.h"
+
+UCLASS()
+class UAnalyticsMulticastSettings
+	: public UAnalyticsSettingsBase
+{
+	GENERATED_UCLASS_BODY()
+
+	/** The list of analytics providers to forward analytics events to */
+	UPROPERTY(EditAnywhere, DisplayName="Release Providers", Category=Multicast, meta=(ConfigRestartRequired=true))
+	TArray<FString> ReleaseMulticastProviders;
+
+	/** The list of analytics providers to forward analytics events to */
+	UPROPERTY(EditAnywhere, DisplayName="Debug Providers", Category=Multicast, meta=(ConfigRestartRequired=true))
+	TArray<FString> DebugMulticastProviders;
+
+	/** The list of analytics providers to forward analytics events to */
+	UPROPERTY(EditAnywhere, DisplayName="Test Providers", Category=Multicast, meta=(ConfigRestartRequired=true))
+	TArray<FString> TestMulticastProviders;
+
+	/** The list of analytics providers to forward analytics events to */
+	UPROPERTY(EditAnywhere, DisplayName="Development Providers", Category=Multicast, meta=(ConfigRestartRequired=true))
+	TArray<FString> DevelopmentMulticastProviders;
+
+// UAnalyticsSettingsBase interface
+protected:
+	/**
+	 * Provides a mechanism to read the section based information into this UObject's properties
+	 */
+	virtual void ReadConfigSettings();
+	/**
+	 * Provides a mechanism to save this object's properties to the section based ini values
+	 */
+	virtual void WriteConfigSettings();
+
+private:
+	/** Helper to populate an array with the results of the config read */
+	void BuildArrayFromString(const FString& List, TArray<FString>& Array);
+	/** Helper to build the comma delimited string from the array */
+	FString BuildStringFromArray(TArray<FString>& Array);
+};
diff --git a/Engine/Plugins/Runtime/Analytics/AnalyticsMulticast/Source/AnaltyicsMulticastEditor/Private/AnalyticsMulticastEditor.cpp b/Engine/Plugins/Runtime/Analytics/AnalyticsMulticast/Source/AnaltyicsMulticastEditor/Private/AnalyticsMulticastEditor.cpp
new file mode 100644
--- /dev/null
+++ b/Engine/Plugins/Runtime/Analytics/AnalyticsMulticast/Source/AnaltyicsMulticastEditor/Private/AnalyticsMulticastEditor.cpp
@@ -0,0 +1,88 @@
+// Copyright 1998-2015 Epic Games, Inc. All Rights Reserved.
+
+#include "AnalyticsMulticastEditorPrivatePCH.h"
+
+#include "AnalyticsMulticastEditor.h"
+#include "AnalyticsMulticastSettings.h"
+
+IMPLEMENT_MODULE( FAnalyticsMulticastEditorModule, AnalyticsMulticastEditor );
+
+#define LOCTEXT_NAMESPACE "AnalyticsMulticast"
+
+void FAnalyticsMulticastEditorModule::StartupModule()
+{
+}
+
+void FAnalyticsMulticastEditorModule::ShutdownModule()
+{
+}
+
+UAnalyticsMulticastSettings::UAnalyticsMulticastSettings(const FObjectInitializer& ObjectInitializer)
+	: Super(ObjectInitializer)
+{
+	SettingsDisplayName = LOCTEXT("SettingsDisplayName", "Multicast");
+	SettingsTooltip = LOCTEXT("SettingsTooltip", "Configuration settings for the Multicast Analytics Provider");
+}
+
+void UAnalyticsMulticastSettings::ReadConfigSettings()
+{
+	FString ProviderList = FAnalytics::Get().GetConfigValueFromIni(GetIniName(), GetReleaseIniSection(), TEXT("ProviderModuleNames"), true);
+	BuildArrayFromString(ProviderList, ReleaseMulticastProviders);
+
+	FString ProviderList2 = FAnalytics::Get().GetConfigValueFromIni(GetIniName(), GetTestIniSection(), TEXT("ProviderModuleNames"), true);
+	if (ProviderList2.Len() == 0)
+	{
+		ProviderList2 = ProviderList;
+	}
+	BuildArrayFromString(ProviderList2, TestMulticastProviders);
+
+	FString ProviderList3 = FAnalytics::Get().GetConfigValueFromIni(GetIniName(), GetDebugIniSection(), TEXT("ProviderModuleNames"), true);
+	if (ProviderList3.Len() == 0)
+	{
+		ProviderList3 = ProviderList;
+	}
+	BuildArrayFromString(ProviderList3, DebugMulticastProviders);
+
+	FString ProviderList4 = FAnalytics::Get().GetConfigValueFromIni(GetIniName(), GetDevelopmentIniSection(), TEXT("ProviderModuleNames"), true);
+	if (ProviderList4.Len() == 0)
+	{
+		ProviderList4 = ProviderList;
+	}
+	BuildArrayFromString(ProviderList4, DevelopmentMulticastProviders);
+}
+
+void UAnalyticsMulticastSettings::WriteConfigSettings()
+{
+	FString ProviderList = BuildStringFromArray(ReleaseMulticastProviders);
+	FAnalytics::Get().WriteConfigValueToIni(GetIniName(), GetReleaseIniSection(), TEXT("ProviderModuleNames"), ProviderList);
+
+	ProviderList = BuildStringFromArray(TestMulticastProviders);
+	FAnalytics::Get().WriteConfigValueToIni(GetIniName(), GetTestIniSection(), TEXT("ProviderModuleNames"), ProviderList);
+
+	ProviderList = BuildStringFromArray(DebugMulticastProviders);
+	FAnalytics::Get().WriteConfigValueToIni(GetIniName(), GetDebugIniSection(), TEXT("ProviderModuleNames"), ProviderList);
+
+	ProviderList = BuildStringFromArray(DevelopmentMulticastProviders);
+	FAnalytics::Get().WriteConfigValueToIni(GetIniName(), GetDevelopmentIniSection(), TEXT("ProviderModuleNames"), ProviderList);
+}
+
+void UAnalyticsMulticastSettings::BuildArrayFromString(const FString& List, TArray<FString>& Array)
+{
+	Array.Empty();
+	List.ParseIntoArray(Array, TEXT(","), true);
+}
+
+FString UAnalyticsMulticastSettings::BuildStringFromArray(TArray<FString>& Array)
+{
+	FString List;
+
+	for (auto ProviderName : Array)
+	{
+		if (List.Len() > 0)
+		{
+			List += TEXT(",");
+		}
+		List += ProviderName;
+	}
+	return List;
+}
diff --git a/Engine/Plugins/Runtime/Analytics/AnalyticsMulticast/Source/AnaltyicsMulticastEditor/Private/AnalyticsMulticastEditorPrivatePCH.h b/Engine/Plugins/Runtime/Analytics/AnalyticsMulticast/Source/AnaltyicsMulticastEditor/Private/AnalyticsMulticastEditorPrivatePCH.h
new file mode 100644
--- /dev/null
+++ b/Engine/Plugins/Runtime/Analytics/AnalyticsMulticast/Source/AnaltyicsMulticastEditor/Private/AnalyticsMulticastEditorPrivatePCH.h
@@ -0,0 +1,6 @@
+// Copyright 1998-2015 Epic Games, Inc. All Rights Reserved.
+
+// You should place include statements to your module's private header files here.  You only need to
+// add includes for headers that are used in most of your module's source files though.
+
+#include "Analytics.h"
diff --git a/Engine/Plugins/Runtime/Analytics/AnalyticsMulticast/Source/AnaltyicsMulticastEditor/Public/AnalyticsMulticastEditor.h b/Engine/Plugins/Runtime/Analytics/AnalyticsMulticast/Source/AnaltyicsMulticastEditor/Public/AnalyticsMulticastEditor.h
new file mode 100644
--- /dev/null
+++ b/Engine/Plugins/Runtime/Analytics/AnalyticsMulticast/Source/AnaltyicsMulticastEditor/Public/AnalyticsMulticastEditor.h
@@ -0,0 +1,12 @@
+// Copyright 1998-2015 Epic Games, Inc. All Rights Reserved.
+
+#pragma once
+
+class FAnalyticsMulticastEditorModule :
+	public IModuleInterface
+{
+private:
+	virtual void StartupModule() override;
+	virtual void ShutdownModule() override;
+};
+
diff --git a/Engine/Plugins/Runtime/Analytics/Apsalar/Apsalar.uplugin b/Engine/Plugins/Runtime/Analytics/Apsalar/Apsalar.uplugin
new file mode 100644
--- /dev/null
+++ b/Engine/Plugins/Runtime/Analytics/Apsalar/Apsalar.uplugin
@@ -0,0 +1,27 @@
+{
+	"FileVersion" : 3,
+	
+	"FriendlyName" : "Apsalar Analytics Provider",
+	"Version" : 2,
+	"VersionName" : "2.0",
+	"CreatedBy" : "Epic Games, Inc.",
+	"CreatedByURL" : "http://epicgames.com",
+	"DocsURL" : "https://docs.unrealengine.com/latest/INT/Gameplay/Analytics/index.html",
+	"EngineVersion" : "4.9.0",
+	"Description" : "Apsalar Analytics Provider",
+	"Category" : "Analytics",
+	
+	"Modules" :
+	[
+        {
+            "Name": "ApsalarEditor",
+            "Type": "Editor"
+        },
+		{
+			"Name" : "IOSApsalar",
+			"Type" : "Runtime",
+			"LoadingPhase" : "PreDefault",
+			"WhitelistPlatforms" : [ "IOS" ]
+		}
+	]
+}
\ No newline at end of file
diff --git a/Engine/Plugins/Runtime/Analytics/Apsalar/IOSApsalar/IOSApsalar.uplugin b/Engine/Plugins/Runtime/Analytics/Apsalar/IOSApsalar/IOSApsalar.uplugin
deleted file mode 100644
diff --git a/Engine/Plugins/Runtime/Analytics/Apsalar/IOSApsalar/Source/IOSApsalar/IOSApsalar.Build.cs b/Engine/Plugins/Runtime/Analytics/Apsalar/IOSApsalar/Source/IOSApsalar/IOSApsalar.Build.cs
deleted file mode 100644
diff --git a/Engine/Plugins/Runtime/Analytics/Apsalar/IOSApsalar/Source/IOSApsalar/Private/IOSApsalar.cpp b/Engine/Plugins/Runtime/Analytics/Apsalar/IOSApsalar/Source/IOSApsalar/Private/IOSApsalar.cpp
deleted file mode 100644
diff --git a/Engine/Plugins/Runtime/Analytics/Apsalar/IOSApsalar/Source/IOSApsalar/Private/IOSApsalarPrivatePCH.h b/Engine/Plugins/Runtime/Analytics/Apsalar/IOSApsalar/Source/IOSApsalar/Private/IOSApsalarPrivatePCH.h
deleted file mode 100644
diff --git a/Engine/Plugins/Runtime/Analytics/Apsalar/IOSApsalar/Source/IOSApsalar/Private/IOSApsalarProvider.h b/Engine/Plugins/Runtime/Analytics/Apsalar/IOSApsalar/Source/IOSApsalar/Private/IOSApsalarProvider.h
deleted file mode 100644
diff --git a/Engine/Plugins/Runtime/Analytics/Apsalar/IOSApsalar/Source/IOSApsalar/Public/IOSApsalar.h b/Engine/Plugins/Runtime/Analytics/Apsalar/IOSApsalar/Source/IOSApsalar/Public/IOSApsalar.h
deleted file mode 100644
diff --git a/Engine/Plugins/Runtime/Analytics/Apsalar/Source/ApsalarEditor/ApsalarEditor.Build.cs b/Engine/Plugins/Runtime/Analytics/Apsalar/Source/ApsalarEditor/ApsalarEditor.Build.cs
new file mode 100644
--- /dev/null
+++ b/Engine/Plugins/Runtime/Analytics/Apsalar/Source/ApsalarEditor/ApsalarEditor.Build.cs
@@ -0,0 +1,29 @@
+// Copyright 1998-2015 Epic Games, Inc. All Rights Reserved.
+
+using UnrealBuildTool;
+
+public class ApsalarEditor : ModuleRules
+{
+    public ApsalarEditor(TargetInfo Target)
+	{
+        PrivateIncludePaths.Add("ApsalarEditor/Private");
+
+		PrivateDependencyModuleNames.AddRange(
+			new string[] {
+				"Core",
+				"CoreUObject",
+                "Analytics",
+                "AnalyticsVisualEditing",
+                "Engine",
+				"Projects"
+			}
+			);
+
+		PrivateIncludePathModuleNames.AddRange(
+			new string[]
+            {
+				"Settings"
+			}
+		);
+	}
+}
diff --git a/Engine/Plugins/Runtime/Analytics/Apsalar/Source/ApsalarEditor/Classes/ApsalarSettings.h b/Engine/Plugins/Runtime/Analytics/Apsalar/Source/ApsalarEditor/Classes/ApsalarSettings.h
new file mode 100644
--- /dev/null
+++ b/Engine/Plugins/Runtime/Analytics/Apsalar/Source/ApsalarEditor/Classes/ApsalarSettings.h
@@ -0,0 +1,61 @@
+// Copyright 1998-2015 Epic Games, Inc. All Rights Reserved.
+
+#pragma once
+
+#include "AnalyticsSettings.h"
+#include "ApsalarSettings.generated.h"
+
+USTRUCT()
+struct FApsalarAnalyticsConfigSetting
+{
+	GENERATED_USTRUCT_BODY()
+
+	UPROPERTY(EditAnywhere, Category=Apsalar, meta=(ConfigRestartRequired=true))
+	FString ApiKey;
+
+	UPROPERTY(EditAnywhere, Category=Apsalar, meta=(ConfigRestartRequired=true))
+	FString ApiSecret;
+
+	UPROPERTY(EditAnywhere, Category=Apsalar, meta=(ConfigRestartRequired=true))
+	int32 SendInterval;
+
+	UPROPERTY(EditAnywhere, Category=Apsalar, meta=(ConfigRestartRequired=true))
+	int32 MaxBufferSize;
+
+	UPROPERTY(EditAnywhere, Category=Apsalar, meta=(ConfigRestartRequired=true))
+	bool ManuallyReportRevenue;
+};
+
+UCLASS()
+class UApsalarSettings
+	: public UAnalyticsSettingsBase
+{
+	GENERATED_UCLASS_BODY()
+
+	UPROPERTY(EditAnywhere, Category=Apsalar, meta=(ConfigRestartRequired=true))
+	FApsalarAnalyticsConfigSetting Release;
+
+	UPROPERTY(EditAnywhere, Category=Apsalar, meta=(ConfigRestartRequired=true))
+	FApsalarAnalyticsConfigSetting Debug;
+
+	UPROPERTY(EditAnywhere, Category=Apsalar, meta=(ConfigRestartRequired=true))
+	FApsalarAnalyticsConfigSetting Test;
+
+	UPROPERTY(EditAnywhere, Category=Apsalar, meta=(ConfigRestartRequired=true))
+	FApsalarAnalyticsConfigSetting Development;
+
+// UAnalyticsSettingsBase interface
+protected:
+	/**
+	 * Provides a mechanism to read the section based information into this UObject's properties
+	 */
+	virtual void ReadConfigSettings();
+	/**
+	 * Provides a mechanism to save this object's properties to the section based ini values
+	 */
+	virtual void WriteConfigSettings();
+
+private:
+	void ReadConfigStruct(const FString& Section, FApsalarAnalyticsConfigSetting& Dest, FApsalarAnalyticsConfigSetting* Default = nullptr);
+	void WriteConfigStruct(const FString& Section, FApsalarAnalyticsConfigSetting& Source);
+};
diff --git a/Engine/Plugins/Runtime/Analytics/Apsalar/Source/ApsalarEditor/Private/ApsalarEditor.cpp b/Engine/Plugins/Runtime/Analytics/Apsalar/Source/ApsalarEditor/Private/ApsalarEditor.cpp
new file mode 100644
--- /dev/null
+++ b/Engine/Plugins/Runtime/Analytics/Apsalar/Source/ApsalarEditor/Private/ApsalarEditor.cpp
@@ -0,0 +1,99 @@
+// Copyright 1998-2015 Epic Games, Inc. All Rights Reserved.
+
+#include "ApsalarEditorPrivatePCH.h"
+
+#include "ApsalarEditor.h"
+#include "ApsalarSettings.h"
+
+IMPLEMENT_MODULE( FApsalarEditorModule, ApsalarEditorModule );
+
+#define LOCTEXT_NAMESPACE "Apsalar"
+
+void FApsalarEditorModule::StartupModule()
+{
+}
+
+void FApsalarEditorModule::ShutdownModule()
+{
+}
+
+UApsalarSettings::UApsalarSettings(const FObjectInitializer& ObjectInitializer)
+	: Super(ObjectInitializer)
+{
+	SettingsDisplayName = LOCTEXT("SettingsDisplayName", "Apsalar");
+	SettingsTooltip = LOCTEXT("SettingsTooltip", "Apsalar configuration settings");
+}
+
+void UApsalarSettings::ReadConfigSettings()
+{
+	ReadConfigStruct(GetReleaseIniSection(), Release);
+	ReadConfigStruct(GetTestIniSection(), Test, &Release);
+	ReadConfigStruct(GetDebugIniSection(), Debug, &Release);
+	ReadConfigStruct(GetDevelopmentIniSection(), Development, &Release);
+}
+
+void UApsalarSettings::WriteConfigSettings()
+{
+	WriteConfigStruct(GetReleaseIniSection(), Release);
+	WriteConfigStruct(GetTestIniSection(), Test);
+	WriteConfigStruct(GetDebugIniSection(), Debug);
+	WriteConfigStruct(GetDevelopmentIniSection(), Development);
+}
+
+void UApsalarSettings::ReadConfigStruct(const FString& Section, FApsalarAnalyticsConfigSetting& Dest, FApsalarAnalyticsConfigSetting* Default)
+{
+	const FString Key = FAnalytics::Get().GetConfigValueFromIni(GetIniName(), Section, TEXT("ApiKey"), true);
+	if (Key.Len() == 0 && Default != nullptr)
+	{
+		Dest.ApiKey = Default->ApiKey;
+	}
+	else
+	{
+		Dest.ApiKey = Key;
+	}
+	const FString Secret = FAnalytics::Get().GetConfigValueFromIni(GetIniName(), Section, TEXT("ApiSecret"), true);
+	if (Secret.Len() == 0 && Default != nullptr)
+	{
+		Dest.ApiSecret = Default->ApiSecret;
+	}
+	else
+	{
+		Dest.ApiSecret = Secret;
+	}
+	const FString SendInterval = FAnalytics::Get().GetConfigValueFromIni(GetIniName(), Section, TEXT("SendInterval"), false);
+	if (SendInterval.Len() == 0 && Default != nullptr)
+	{
+		Dest.SendInterval = Default->SendInterval;
+	}
+	else
+	{
+		Dest.SendInterval = FCString::Atoi(*SendInterval);
+	}
+	const FString MaxBufferSize = FAnalytics::Get().GetConfigValueFromIni(GetIniName(), Section, TEXT("MaxBufferSize"), false);
+	if (MaxBufferSize.Len() == 0 && Default != nullptr)
+	{
+		Dest.MaxBufferSize = Default->MaxBufferSize;
+	}
+	else
+	{
+		Dest.MaxBufferSize = FCString::Atoi(*MaxBufferSize);
+	}
+	const FString ManuallyReportRevenue = FAnalytics::Get().GetConfigValueFromIni(GetIniName(), Section, TEXT("ManuallyReportRevenue"), false);
+	if (ManuallyReportRevenue.Len() == 0 && Default != nullptr)
+	{
+		Dest.ManuallyReportRevenue = Default->ManuallyReportRevenue;
+	}
+	else
+	{
+		Dest.ManuallyReportRevenue = ManuallyReportRevenue.Compare(TEXT("true"), ESearchCase::IgnoreCase) == 0;
+	}
+}
+
+void UApsalarSettings::WriteConfigStruct(const FString& Section, FApsalarAnalyticsConfigSetting& Source)
+{
+	FAnalytics::Get().WriteConfigValueToIni(GetIniName(), Section, TEXT("ApiKey"), Source.ApiKey);
+	FAnalytics::Get().WriteConfigValueToIni(GetIniName(), Section, TEXT("ApiSecret"), Source.ApiSecret);
+	FAnalytics::Get().WriteConfigValueToIni(GetIniName(), Section, TEXT("SendInterval"), FString::Printf(TEXT("%d"), Source.SendInterval));
+	FAnalytics::Get().WriteConfigValueToIni(GetIniName(), Section, TEXT("MaxBufferSize"), FString::Printf(TEXT("%d"), Source.MaxBufferSize));
+	FAnalytics::Get().WriteConfigValueToIni(GetIniName(), Section, TEXT("ManuallyReportRevenue"), Source.ManuallyReportRevenue ? TEXT("true") : TEXT("false"));
+}
diff --git a/Engine/Plugins/Runtime/Analytics/Apsalar/Source/ApsalarEditor/Private/ApsalarEditorPrivatePCH.h b/Engine/Plugins/Runtime/Analytics/Apsalar/Source/ApsalarEditor/Private/ApsalarEditorPrivatePCH.h
new file mode 100644
--- /dev/null
+++ b/Engine/Plugins/Runtime/Analytics/Apsalar/Source/ApsalarEditor/Private/ApsalarEditorPrivatePCH.h
@@ -0,0 +1,6 @@
+// Copyright 1998-2015 Epic Games, Inc. All Rights Reserved.
+
+// You should place include statements to your module's private header files here.  You only need to
+// add includes for headers that are used in most of your module's source files though.
+
+#include "Analytics.h"
diff --git a/Engine/Plugins/Runtime/Analytics/Apsalar/Source/ApsalarEditor/Public/ApsalarEditor.h b/Engine/Plugins/Runtime/Analytics/Apsalar/Source/ApsalarEditor/Public/ApsalarEditor.h
new file mode 100644
--- /dev/null
+++ b/Engine/Plugins/Runtime/Analytics/Apsalar/Source/ApsalarEditor/Public/ApsalarEditor.h
@@ -0,0 +1,12 @@
+// Copyright 1998-2015 Epic Games, Inc. All Rights Reserved.
+
+#pragma once
+
+class FApsalarEditorModule :
+	public IModuleInterface
+{
+private:
+	virtual void StartupModule() override;
+	virtual void ShutdownModule() override;
+};
+
diff --git a/Engine/Plugins/Runtime/Analytics/Apsalar/Source/IOSApsalar/IOSApsalar.Build.cs b/Engine/Plugins/Runtime/Analytics/Apsalar/Source/IOSApsalar/IOSApsalar.Build.cs
new file mode 100644
--- /dev/null
+++ b/Engine/Plugins/Runtime/Analytics/Apsalar/Source/IOSApsalar/IOSApsalar.Build.cs
@@ -0,0 +1,69 @@
+// Copyright 1998-2015 Epic Games, Inc. All Rights Reserved.
+
+namespace UnrealBuildTool.Rules
+{
+	public class IOSApsalar : ModuleRules
+	{
+		public IOSApsalar( TargetInfo Target )
+		{
+			BinariesSubFolder = "NotForLicensees";
+
+			PublicDependencyModuleNames.AddRange(
+				new string[]
+				{
+					"Core",
+					// ... add other public dependencies that you statically link with here ...
+				}
+				);
+
+			PrivateDependencyModuleNames.AddRange(
+				new string[]
+				{
+					"Analytics",
+					// ... add private dependencies that you statically link with here ...
+				}
+				);
+
+            PublicIncludePathModuleNames.Add("Analytics");
+
+			PublicFrameworks.AddRange(
+				new string[] {
+					"CoreTelephony",
+					"SystemConfiguration",
+					"UIKit",
+					"Foundation",
+					"CoreGraphics",
+					"MobileCoreServices",
+					"StoreKit",
+					"CFNetwork",
+					"CoreData",
+					"Security",
+					"CoreLocation"
+				});
+
+            PublicAdditionalLibraries.AddRange(
+				new string[] {
+					"sqlite3",
+					"z"
+				});
+
+            bool bHasApsalarSDK =
+                (System.IO.Directory.Exists(System.IO.Path.Combine(UEBuildConfiguration.UEThirdPartySourceDirectory, "Apsalar")) &&
+                  System.IO.Directory.Exists(System.IO.Path.Combine(UEBuildConfiguration.UEThirdPartySourceDirectory, "Apsalar", "IOS"))) ||
+                (System.IO.Directory.Exists(System.IO.Path.Combine(UEBuildConfiguration.UEThirdPartySourceDirectory, "NotForLicensees")) &&
+                  System.IO.Directory.Exists(System.IO.Path.Combine(UEBuildConfiguration.UEThirdPartySourceDirectory, "NotForLicensees", "Apsalar")) &&
+                  System.IO.Directory.Exists(System.IO.Path.Combine(UEBuildConfiguration.UEThirdPartySourceDirectory, "NotForLicensees", "Apsalar", "IOS")));
+            if (bHasApsalarSDK)
+            {
+                PublicIncludePaths.Add(UEBuildConfiguration.UEThirdPartySourceDirectory + "NotForLicensees/Apsalar/IOS/");
+                PublicAdditionalLibraries.Add(UEBuildConfiguration.UEThirdPartySourceDirectory + "NotForLicensees/Apsalar/IOS/libApsalar.a");
+
+                Definitions.Add("WITH_APSALAR=1");
+            }
+            else
+            {
+                Definitions.Add("WITH_APSALAR=0");
+            }
+        }
+	}
+}
\ No newline at end of file
diff --git a/Engine/Plugins/Runtime/Analytics/Apsalar/Source/IOSApsalar/Private/IOSApsalar.cpp b/Engine/Plugins/Runtime/Analytics/Apsalar/Source/IOSApsalar/Private/IOSApsalar.cpp
new file mode 100644
--- /dev/null
+++ b/Engine/Plugins/Runtime/Analytics/Apsalar/Source/IOSApsalar/Private/IOSApsalar.cpp
@@ -0,0 +1,296 @@
+// Copyright 1998-2015 Epic Games, Inc. All Rights Reserved.
+
+#include "IOSApsalarPrivatePCH.h"
+
+DEFINE_LOG_CATEGORY_STATIC(LogAnalytics, Display, All);
+
+IMPLEMENT_MODULE( FAnalyticsIOSApsalar, IOSApsalar )
+
+TSharedPtr<IAnalyticsProvider> FAnalyticsProviderApsalar::Provider;
+
+void FAnalyticsIOSApsalar::StartupModule()
+{
+}
+
+void FAnalyticsIOSApsalar::ShutdownModule()
+{
+	FAnalyticsProviderApsalar::Destroy();
+}
+
+TSharedPtr<IAnalyticsProvider> FAnalyticsIOSApsalar::CreateAnalyticsProvider(const FAnalytics::FProviderConfigurationDelegate& GetConfigValue) const
+{
+	if (GetConfigValue.IsBound())
+	{
+		const FString Key = GetConfigValue.Execute(TEXT("ApiKey"), true);
+		const FString Secret = GetConfigValue.Execute(TEXT("ApiSecret"), true);
+		const FString SendInterval = GetConfigValue.Execute(TEXT("SendInterval"), false);
+		const FString MaxBufferSize = GetConfigValue.Execute(TEXT("MaxBufferSize"), false);
+		const FString ManuallyReportRevenue = GetConfigValue.Execute(TEXT("ManuallyReportRevenue"), false);
+		const bool bWantsManualRevenueReporting = ManuallyReportRevenue.Compare(TEXT("true"), ESearchCase::IgnoreCase) == 0;
+		return FAnalyticsProviderApsalar::Create(Key, Secret, FCString::Atoi(*SendInterval), FCString::Atoi(*MaxBufferSize), bWantsManualRevenueReporting);
+	}
+	else
+	{
+		UE_LOG(LogAnalytics, Warning, TEXT("IOSApsalar::CreateAnalyticsProvider called with an unbound delegate"));
+	}
+	return nullptr;
+}
+
+#if !UE_BUILD_SHIPPING
+	/**
+	 * Verifies that the event name matches Apsalar's 32 char limit and warns if it doesn't
+	 */
+	inline void CheckEventNameLen(const FString& EventName)
+	{
+		const int32 Length = EventName.Len();
+		if (Length > 32)
+		{
+			UE_LOG(LogAnalytics, Warning, TEXT("Apsalar event name is too long: %s and will be truncated by Apsalar. 32 character max limit."), *EventName);
+		}
+		else if (Length == 0)
+		{
+			UE_LOG(LogAnalytics, Warning, TEXT("Apsalar event name is empty!"));
+		}
+	}
+
+	#define WarnIfEventNameIsWrongLength(x) CheckEventNameLen(x)
+#else
+	// Have these compile out
+	#define WarnIfEventNameIsWrongLength(x)
+#endif
+
+// Provider
+
+FAnalyticsProviderApsalar::FAnalyticsProviderApsalar(const FString Key, const FString Secret, const int32 SendInterval, const int32 MaxBufferSize, const bool bWantsManualRevenueReporting) :
+	ApiKey(Key),
+	ApiSecret(Secret)
+{
+#if WITH_APSALAR
+	// If this is zero, then use the Apsalar default buffer size
+	if (MaxBufferSize > 0)
+	{
+		[Apsalar setBufferLimit:MaxBufferSize];
+	}
+
+	// If this is zero, then use the Apsalar default time period
+	if (SendInterval > 0)
+	{
+		[Apsalar setBatchInterval:(int)SendInterval];
+	}
+
+	// Disable the auto reporting of revenue if they want to manually report it
+	if (bWantsManualRevenueReporting)
+	{
+		[Apsalar setAllowAutoIAPComplete:NO];
+	}
+#else
+	UE_LOG(LogAnalytics, Warning, TEXT("WITH_APSALAR=0. Are you missing the SDK?"));
+#endif
+}
+
+FAnalyticsProviderApsalar::~FAnalyticsProviderApsalar()
+{
+	if (bHasSessionStarted)
+	{
+		EndSession();
+	}
+}
+
+bool FAnalyticsProviderApsalar::StartSession(const TArray<FAnalyticsEventAttribute>& Attributes)
+{
+#if WITH_APSALAR
+	NSString* ApiKeyStr = [NSString stringWithFString:ApiKey];
+	NSString* ApiSecretStr = [NSString stringWithFString:ApiSecret];
+	const int32 AttrCount = Attributes.Num();
+
+	if (!bHasSessionStarted)
+	{
+		[Apsalar startSession:ApiKeyStr withKey:ApiSecretStr];
+		if (AttrCount > 0)
+		{
+			RecordEvent(TEXT("SessionAttributes"), Attributes);
+		}
+
+		UE_LOG(LogAnalytics, Display, TEXT("IOSApsalar::StartSession(%d attributes)"), AttrCount);
+	}
+	else
+	{
+		[Apsalar reStartSession:ApiKeyStr withKey:ApiSecretStr];
+		if (AttrCount > 0)
+		{
+			RecordEvent(TEXT("SessionAttributes"), Attributes);
+		}
+
+		UE_LOG(LogAnalytics, Display, TEXT("IOSApsalar::RestartSession(%d attributes)"), AttrCount);
+	}
+	bHasSessionStarted = ([Apsalar sessionStarted] == YES) ? true : false;
+	return bHasSessionStarted;
+#else
+	UE_LOG(LogAnalytics, Warning, TEXT("WITH_APSALAR=0. Are you missing the SDK?"));
+	return false;
+#endif
+}
+
+void FAnalyticsProviderApsalar::EndSession()
+{
+#if WITH_APSALAR
+	[Apsalar endSession];
+	bHasSessionStarted = false;
+
+	UE_LOG(LogAnalytics, Display, TEXT("IOSApsalar::EndSession"));
+#else
+	UE_LOG(LogAnalytics, Warning, TEXT("WITH_APSALAR=0. Are you missing the SDK?"));
+#endif
+}
+
+void FAnalyticsProviderApsalar::FlushEvents()
+{
+#if WITH_APSALAR
+	[Apsalar sendAllBatches];
+
+	UE_LOG(LogAnalytics, Display, TEXT("IOSApsalar::FlushEvents"));
+#else
+	UE_LOG(LogAnalytics, Warning, TEXT("WITH_APSALAR=0. Are you missing the SDK?"));
+#endif
+}
+
+void FAnalyticsProviderApsalar::SetUserID(const FString& InUserID)
+{
+#if WITH_APSALAR
+	// Ignored
+	UE_LOG(LogAnalytics, Display, TEXT("IOSApsalar::SetUserID - ignoring call"));
+#else
+	UE_LOG(LogAnalytics, Warning, TEXT("WITH_APSALAR=0. Are you missing the SDK?"));
+#endif
+}
+
+FString FAnalyticsProviderApsalar::GetUserID() const
+{
+#if WITH_APSALAR
+	NSString* ApsalarId = [Apsalar apsalarID];
+	const FString ConvertedId(ApsalarId);
+
+	UE_LOG(LogAnalytics, Display, TEXT("IOSApsalar::GetUserID - returning the id as '%s'"), *ConvertedId);
+
+	return ConvertedId;
+#else
+	UE_LOG(LogAnalytics, Warning, TEXT("WITH_APSALAR=0. Are you missing the SDK?"));
+	return FString();
+#endif
+}
+
+FString FAnalyticsProviderApsalar::GetSessionID() const
+{
+#if WITH_APSALAR
+	NSString* Id = [Apsalar sessionID];
+	FString ConvertedId(Id);
+
+	UE_LOG(LogAnalytics, Display, TEXT("IOSApsalar::GetSessionID - returning the id as '%s'"), *ConvertedId);
+
+	return ConvertedId;
+#else
+	UE_LOG(LogAnalytics, Warning, TEXT("WITH_APSALAR=0. Are you missing the SDK?"));
+	return FString();
+#endif
+}
+
+bool FAnalyticsProviderApsalar::SetSessionID(const FString& InSessionID)
+{
+#if WITH_APSALAR
+	// Ignored
+	UE_LOG(LogAnalytics, Display, TEXT("IOSApsalar::SetSessionID - ignoring call"));
+	return false;
+#else
+	UE_LOG(LogAnalytics, Warning, TEXT("WITH_APSALAR=0. Are you missing the SDK?"));
+	return false;
+#endif
+}
+
+void FAnalyticsProviderApsalar::RecordEvent(const FString& EventName, const TArray<FAnalyticsEventAttribute>& Attributes)
+{
+#if WITH_APSALAR
+	WarnIfEventNameIsWrongLength(EventName);
+
+	if (EventName.Len() > 0)
+	{
+		NSString* ConvertedEventName = [NSString stringWithFString : EventName];
+		const int32 AttrCount = Attributes.Num();
+		if (AttrCount > 0)
+		{
+			// Convert the event attributes into a dictionary object
+			NSDictionary* AttributesDict = [NSMutableDictionary dictionaryWithCapacity:AttrCount];
+			for	(auto Attr : Attributes)
+			{
+				NSString* AttrName = [NSString stringWithFString : Attr.AttrName];
+				NSString* AttrValue = [NSString stringWithFString : Attr.AttrValue];
+				[AttributesDict setValue:AttrValue forKey:AttrName];
+			}
+			[Apsalar event:ConvertedEventName withArgs:AttributesDict];
+		}
+		else
+		{
+			// Log an event with no payload
+			[Apsalar event:ConvertedEventName];
+		}
+		UE_LOG(LogAnalytics, Display, TEXT("IOSApsalar::RecordEvent('%s', %d attributes)"), *EventName, AttrCount);
+	}
+#else
+	UE_LOG(LogAnalytics, Warning, TEXT("WITH_APSALAR=0. Are you missing the SDK?"));
+#endif
+}
+
+void FAnalyticsProviderApsalar::RecordItemPurchase(const FString& ItemId, const FString& Currency, int PerItemCost, int ItemQuantity)
+{
+#if WITH_APSALAR
+	NSString* EventName = @"Item Purchase";
+	// Build the dictionary
+	NSDictionary* AttributesDict = [NSMutableDictionary dictionaryWithCapacity:4];
+	[AttributesDict setValue:[NSString stringWithFString:ItemId] forKey:@"ItemId"];
+	[AttributesDict setValue:[NSString stringWithFString:Currency] forKey:@"Currency"];
+	[AttributesDict setValue:[NSString stringWithFormat:@"%d", PerItemCost] forKey:@"PerItemCost"];
+	[AttributesDict setValue:[NSString stringWithFormat:@"%d", ItemQuantity] forKey:@"ItemQuantity"];
+	// Send the event
+	[Apsalar event:EventName withArgs:AttributesDict];
+
+	UE_LOG(LogAnalytics, Display, TEXT("IOSApsalar::RecordItemPurchase('%s', '%s', %d, %d)"), *ItemId, *Currency, PerItemCost, ItemQuantity);
+#else
+	UE_LOG(LogAnalytics, Warning, TEXT("WITH_APSALAR=0. Are you missing the SDK?"));
+#endif
+}
+
+void FAnalyticsProviderApsalar::RecordCurrencyPurchase(const FString& GameCurrencyType, int GameCurrencyAmount, const FString& RealCurrencyType, float RealMoneyCost, const FString& PaymentProvider)
+{
+#if WITH_APSALAR
+	NSString* EventName = @"Currency Purchase";
+	// Build the dictionary
+	NSDictionary* AttributesDict = [NSMutableDictionary dictionaryWithCapacity:5];
+	[AttributesDict setValue:[NSString stringWithFString:GameCurrencyType] forKey:@"GameCurrencyType"];
+	[AttributesDict setValue:[NSString stringWithFormat:@"%d", GameCurrencyAmount] forKey:@"GameCurrencyAmount"];
+	[AttributesDict setValue:[NSString stringWithFString:RealCurrencyType] forKey:@"RealCurrencyType"];
+	[AttributesDict setValue:[NSString stringWithFormat:@"%.02f", RealMoneyCost] forKey:@"RealMoneyCost"];
+	[AttributesDict setValue:[NSString stringWithFString:PaymentProvider] forKey:@"PaymentProvider"];
+	// Send the event
+	[Apsalar event:EventName withArgs:AttributesDict];
+
+	UE_LOG(LogAnalytics, Display, TEXT("IOSApsalar::RecordCurrencyPurchase('%s', %d, '%s', %.02f, %s)"), *GameCurrencyType, GameCurrencyAmount, *RealCurrencyType, RealMoneyCost, *PaymentProvider);
+#else
+	UE_LOG(LogAnalytics, Warning, TEXT("WITH_APSALAR=0. Are you missing the SDK?"));
+#endif
+}
+
+void FAnalyticsProviderApsalar::RecordCurrencyGiven(const FString& GameCurrencyType, int GameCurrencyAmount)
+{
+#if WITH_APSALAR
+	NSString* EventName = @"Currency Given";
+	// Build the dictionary
+	NSDictionary* AttributesDict = [NSMutableDictionary dictionaryWithCapacity:2];
+	[AttributesDict setValue:[NSString stringWithFString:GameCurrencyType] forKey:@"GameCurrencyType"];
+	[AttributesDict setValue:[NSString stringWithFormat:@"%d", GameCurrencyAmount] forKey:@"GameCurrencyAmount"];
+	// Send the event
+	[Apsalar event:EventName withArgs:AttributesDict];
+
+	UE_LOG(LogAnalytics, Display, TEXT("IOSApsalar::RecordCurrencyGiven('%s', %d)"), *GameCurrencyType, GameCurrencyAmount);
+#else
+	UE_LOG(LogAnalytics, Warning, TEXT("WITH_APSALAR=0. Are you missing the SDK?"));
+#endif
+}
diff --git a/Engine/Plugins/Runtime/Analytics/Apsalar/Source/IOSApsalar/Private/IOSApsalarPrivatePCH.h b/Engine/Plugins/Runtime/Analytics/Apsalar/Source/IOSApsalar/Private/IOSApsalarPrivatePCH.h
new file mode 100644
--- /dev/null
+++ b/Engine/Plugins/Runtime/Analytics/Apsalar/Source/IOSApsalar/Private/IOSApsalarPrivatePCH.h
@@ -0,0 +1,15 @@
+// Copyright 1998-2015 Epic Games, Inc. All Rights Reserved.
+
+#include "Core.h"
+#include "IAnalyticsProviderModule.h"
+#include "IAnalyticsProvider.h"
+
+#include "IOSApsalar.h"
+#include "IOSApsalarProvider.h"
+
+#if WITH_APSALAR
+#import "Apsalar.h"
+#endif
+
+// You should place include statements to your module's private header files here.  You only need to
+// add includes for headers that are used in most of your module's source files though.
diff --git a/Engine/Plugins/Runtime/Analytics/Apsalar/Source/IOSApsalar/Private/IOSApsalarProvider.h b/Engine/Plugins/Runtime/Analytics/Apsalar/Source/IOSApsalar/Private/IOSApsalarProvider.h
new file mode 100644
--- /dev/null
+++ b/Engine/Plugins/Runtime/Analytics/Apsalar/Source/IOSApsalar/Private/IOSApsalarProvider.h
@@ -0,0 +1,50 @@
+// Copyright 1998-2015 Epic Games, Inc. All Rights Reserved.
+
+#pragma once
+
+class FAnalyticsProviderApsalar :
+	public IAnalyticsProvider
+{
+	/** The API key given to you by Apsalar */
+	FString ApiKey;
+	/** The API secret generated for you by Apsalar */
+	FString ApiSecret;
+	/** Tracks whether we need to start the session or restart it */
+	bool bHasSessionStarted;
+
+	/** Singleton for analytics */
+	static TSharedPtr<IAnalyticsProvider> Provider;
+	FAnalyticsProviderApsalar(const FString Key, const FString Secret, const int32 SendInterval, const int32 MaxBufferSize, const bool bWantsManualRevenueReporting);
+
+public:
+	static TSharedPtr<IAnalyticsProvider> Create(const FString Key, const FString Secret, const int32 SendInterval, const int32 MaxBufferSize, const bool bWantsManualRevenueReporting)
+	{
+		if (!Provider.IsValid())
+		{
+			Provider = TSharedPtr<IAnalyticsProvider>(new FAnalyticsProviderApsalar(Key, Secret, SendInterval, MaxBufferSize, bWantsManualRevenueReporting));
+		}
+		return Provider;
+	}
+	static void Destroy()
+	{
+		Provider.Reset();
+	}
+
+	virtual ~FAnalyticsProviderApsalar();
+
+	virtual bool StartSession(const TArray<FAnalyticsEventAttribute>& Attributes) override;
+	virtual void EndSession() override;
+	virtual void FlushEvents() override;
+
+	virtual void SetUserID(const FString& InUserID) override;
+	virtual FString GetUserID() const override;
+
+	virtual FString GetSessionID() const override;
+	virtual bool SetSessionID(const FString& InSessionID) override;
+
+	virtual void RecordEvent(const FString& EventName, const TArray<FAnalyticsEventAttribute>& Attributes) override;
+
+	virtual void RecordItemPurchase(const FString& ItemId, const FString& Currency, int PerItemCost, int ItemQuantity) override;
+	virtual void RecordCurrencyPurchase(const FString& GameCurrencyType, int GameCurrencyAmount, const FString& RealCurrencyType, float RealMoneyCost, const FString& PaymentProvider) override;
+	virtual void RecordCurrencyGiven(const FString& GameCurrencyType, int GameCurrencyAmount) override;
+};
diff --git a/Engine/Plugins/Runtime/Analytics/Apsalar/Source/IOSApsalar/Public/IOSApsalar.h b/Engine/Plugins/Runtime/Analytics/Apsalar/Source/IOSApsalar/Public/IOSApsalar.h
new file mode 100644
--- /dev/null
+++ b/Engine/Plugins/Runtime/Analytics/Apsalar/Source/IOSApsalar/Public/IOSApsalar.h
@@ -0,0 +1,46 @@
+// Copyright 1998-2015 Epic Games, Inc. All Rights Reserved.
+
+#pragma once
+
+#include "Runtime/Analytics/Analytics/Public/Interfaces/IAnalyticsProviderModule.h"
+#include "Core.h"
+
+class IAnalyticsProvider;
+
+/**
+ * The public interface to this module
+ */
+class FAnalyticsIOSApsalar :
+	public IAnalyticsProviderModule
+{
+	//--------------------------------------------------------------------------
+	// Module functionality
+	//--------------------------------------------------------------------------
+public:
+	/**
+	 * Singleton-like access to this module's interface.  This is just for convenience!
+	 * Beware of calling this during the shutdown phase, though.  Your module might have been unloaded already.
+	 *
+	 * @return Returns singleton instance, loading the module on demand if needed
+	 */
+	static inline FAnalyticsIOSApsalar& Get()
+	{
+		return FModuleManager::LoadModuleChecked< FAnalyticsIOSApsalar >( "IOSApsalar" );
+	}
+
+	//--------------------------------------------------------------------------
+	// provider factory functions
+	//--------------------------------------------------------------------------
+public:
+	/**
+	 * IAnalyticsProviderModule interface.
+	 * Creates the analytics provider given a configuration delegate.
+	 * The keys required exactly match the field names in the Config object. 
+	 */
+	virtual TSharedPtr<IAnalyticsProvider> CreateAnalyticsProvider(const FAnalytics::FProviderConfigurationDelegate& GetConfigValue) const override;
+	
+private:
+	virtual void StartupModule() override;
+	virtual void ShutdownModule() override;
+};
+
diff --git a/Engine/Plugins/Runtime/Analytics/Flurry/Flurry.uplugin b/Engine/Plugins/Runtime/Analytics/Flurry/Flurry.uplugin
new file mode 100644
--- /dev/null
+++ b/Engine/Plugins/Runtime/Analytics/Flurry/Flurry.uplugin
@@ -0,0 +1,27 @@
+{
+	"FileVersion" : 3,
+	
+	"FriendlyName" : "Flurry Analytics Provider",
+	"Version" : 1,
+	"VersionName" : "1.0",
+	"CreatedBy" : "Epic Games, Inc.",
+	"CreatedByURL" : "http://epicgames.com",
+	"DocsURL" : "https://docs.unrealengine.com/latest/INT/Gameplay/Analytics/index.html",
+	"EngineVersion" : "4.8.0",
+	"Description" : "Flurry Analytics Provider",
+	"Category" : "Analytics",
+	
+	"Modules" :
+	[
+        {
+            "Name": "FlurryEditor",
+            "Type": "Editor"
+        },
+		{
+			"Name" : "IOSFlurry",
+			"Type" : "Runtime",
+			"LoadingPhase" : "PreDefault",
+			"WhitelistPlatforms" : [ "IOS" ]
+		}
+	]
+}
\ No newline at end of file
diff --git a/Engine/Plugins/Runtime/Analytics/Flurry/IOSFlurry/IOSFlurry.uplugin b/Engine/Plugins/Runtime/Analytics/Flurry/IOSFlurry/IOSFlurry.uplugin
deleted file mode 100644
diff --git a/Engine/Plugins/Runtime/Analytics/Flurry/IOSFlurry/Source/IOSFlurry/IOSFlurry.Build.cs b/Engine/Plugins/Runtime/Analytics/Flurry/IOSFlurry/Source/IOSFlurry/IOSFlurry.Build.cs
deleted file mode 100644
diff --git a/Engine/Plugins/Runtime/Analytics/Flurry/IOSFlurry/Source/IOSFlurry/Private/IOSFlurry.cpp b/Engine/Plugins/Runtime/Analytics/Flurry/IOSFlurry/Source/IOSFlurry/Private/IOSFlurry.cpp
deleted file mode 100644
diff --git a/Engine/Plugins/Runtime/Analytics/Flurry/IOSFlurry/Source/IOSFlurry/Private/IOSFlurryPrivatePCH.h b/Engine/Plugins/Runtime/Analytics/Flurry/IOSFlurry/Source/IOSFlurry/Private/IOSFlurryPrivatePCH.h
deleted file mode 100644
diff --git a/Engine/Plugins/Runtime/Analytics/Flurry/IOSFlurry/Source/IOSFlurry/Private/IOSFlurryProvider.h b/Engine/Plugins/Runtime/Analytics/Flurry/IOSFlurry/Source/IOSFlurry/Private/IOSFlurryProvider.h
deleted file mode 100644
diff --git a/Engine/Plugins/Runtime/Analytics/Flurry/IOSFlurry/Source/IOSFlurry/Public/IOSFlurry.h b/Engine/Plugins/Runtime/Analytics/Flurry/IOSFlurry/Source/IOSFlurry/Public/IOSFlurry.h
deleted file mode 100644
diff --git a/Engine/Plugins/Runtime/Analytics/Flurry/Source/FlurryEditor/Classes/FlurrySettings.h b/Engine/Plugins/Runtime/Analytics/Flurry/Source/FlurryEditor/Classes/FlurrySettings.h
new file mode 100644
--- /dev/null
+++ b/Engine/Plugins/Runtime/Analytics/Flurry/Source/FlurryEditor/Classes/FlurrySettings.h
@@ -0,0 +1,36 @@
+// Copyright 1998-2015 Epic Games, Inc. All Rights Reserved.
+
+#pragma once
+
+#include "AnalyticsSettings.h"
+#include "FlurrySettings.generated.h"
+
+UCLASS()
+class UFlurrySettings
+	: public UAnalyticsSettingsBase
+{
+	GENERATED_UCLASS_BODY()
+
+	UPROPERTY(EditAnywhere, Category=Flurry, meta=(ConfigRestartRequired=true))
+	FString ReleaseApiKey;
+
+	UPROPERTY(EditAnywhere, Category=Flurry, meta=(ConfigRestartRequired=true))
+	FString DebugApiKey;
+
+	UPROPERTY(EditAnywhere, Category=Flurry, meta=(ConfigRestartRequired=true))
+	FString TestApiKey;
+
+	UPROPERTY(EditAnywhere, Category=Flurry, meta=(ConfigRestartRequired=true))
+	FString DevelopmentApiKey;
+
+// UAnalyticsSettingsBase interface
+protected:
+	/**
+	 * Provides a mechanism to read the section based information into this UObject's properties
+	 */
+	virtual void ReadConfigSettings();
+	/**
+	 * Provides a mechanism to save this object's properties to the section based ini values
+	 */
+	virtual void WriteConfigSettings();
+};
diff --git a/Engine/Plugins/Runtime/Analytics/Flurry/Source/FlurryEditor/FlurryEditor.Build.cs b/Engine/Plugins/Runtime/Analytics/Flurry/Source/FlurryEditor/FlurryEditor.Build.cs
new file mode 100644
--- /dev/null
+++ b/Engine/Plugins/Runtime/Analytics/Flurry/Source/FlurryEditor/FlurryEditor.Build.cs
@@ -0,0 +1,29 @@
+// Copyright 1998-2015 Epic Games, Inc. All Rights Reserved.
+
+using UnrealBuildTool;
+
+public class FlurryEditor : ModuleRules
+{
+    public FlurryEditor(TargetInfo Target)
+	{
+        PrivateIncludePaths.Add("FlurryEditor/Private");
+
+		PrivateDependencyModuleNames.AddRange(
+			new string[] {
+				"Core",
+				"CoreUObject",
+                "Analytics",
+                "AnalyticsVisualEditing",
+                "Engine",
+				"Projects"
+			}
+			);
+
+		PrivateIncludePathModuleNames.AddRange(
+			new string[]
+            {
+				"Settings"
+			}
+		);
+	}
+}
diff --git a/Engine/Plugins/Runtime/Analytics/Flurry/Source/FlurryEditor/Private/FlurryEditor.cpp b/Engine/Plugins/Runtime/Analytics/Flurry/Source/FlurryEditor/Private/FlurryEditor.cpp
new file mode 100644
--- /dev/null
+++ b/Engine/Plugins/Runtime/Analytics/Flurry/Source/FlurryEditor/Private/FlurryEditor.cpp
@@ -0,0 +1,60 @@
+// Copyright 1998-2015 Epic Games, Inc. All Rights Reserved.
+
+#include "FlurryEditorPrivatePCH.h"
+
+#include "FlurryEditor.h"
+#include "FlurrySettings.h"
+
+IMPLEMENT_MODULE( FFlurryEditorModule, FlurryEditorModule );
+
+#define LOCTEXT_NAMESPACE "Flurry"
+
+void FFlurryEditorModule::StartupModule()
+{
+}
+
+void FFlurryEditorModule::ShutdownModule()
+{
+}
+
+UFlurrySettings::UFlurrySettings(const FObjectInitializer& ObjectInitializer)
+	: Super(ObjectInitializer)
+{
+	SettingsDisplayName = LOCTEXT("SettingsDisplayName", "Flurry");
+	SettingsTooltip = LOCTEXT("SettingsTooltip", "Flurry analytics configuration settings");
+}
+
+void UFlurrySettings::ReadConfigSettings()
+{
+	FString ReadApiKey = FAnalytics::Get().GetConfigValueFromIni(GetIniName(), GetReleaseIniSection(), TEXT("FlurryApiKey"), true);
+	ReleaseApiKey = ReadApiKey;
+
+	ReadApiKey = FAnalytics::Get().GetConfigValueFromIni(GetIniName(), GetDebugIniSection(), TEXT("FlurryApiKey"), true);
+	if (ReadApiKey.Len() == 0)
+	{
+		ReadApiKey = ReleaseApiKey;
+	}
+	DebugApiKey = ReadApiKey;
+
+	ReadApiKey = FAnalytics::Get().GetConfigValueFromIni(GetIniName(), GetTestIniSection(), TEXT("FlurryApiKey"), true);
+	if (ReadApiKey.Len() == 0)
+	{
+		ReadApiKey = ReleaseApiKey;
+	}
+	TestApiKey = ReadApiKey;
+
+	ReadApiKey = FAnalytics::Get().GetConfigValueFromIni(GetIniName(), GetDevelopmentIniSection(), TEXT("FlurryApiKey"), true);
+	if (ReadApiKey.Len() == 0)
+	{
+		ReadApiKey = ReleaseApiKey;
+	}
+	DevelopmentApiKey = ReadApiKey;
+}
+
+void UFlurrySettings::WriteConfigSettings()
+{
+	FAnalytics::Get().WriteConfigValueToIni(GetIniName(), GetReleaseIniSection(), TEXT("FlurryApiKey"), ReleaseApiKey);
+	FAnalytics::Get().WriteConfigValueToIni(GetIniName(), GetTestIniSection(), TEXT("FlurryApiKey"), TestApiKey);
+	FAnalytics::Get().WriteConfigValueToIni(GetIniName(), GetDebugIniSection(), TEXT("FlurryApiKey"), DebugApiKey);
+	FAnalytics::Get().WriteConfigValueToIni(GetIniName(), GetDevelopmentIniSection(), TEXT("FlurryApiKey"), DevelopmentApiKey);
+}
diff --git a/Engine/Plugins/Runtime/Analytics/Flurry/Source/FlurryEditor/Private/FlurryEditorPrivatePCH.h b/Engine/Plugins/Runtime/Analytics/Flurry/Source/FlurryEditor/Private/FlurryEditorPrivatePCH.h
new file mode 100644
--- /dev/null
+++ b/Engine/Plugins/Runtime/Analytics/Flurry/Source/FlurryEditor/Private/FlurryEditorPrivatePCH.h
@@ -0,0 +1,6 @@
+// Copyright 1998-2015 Epic Games, Inc. All Rights Reserved.
+
+// You should place include statements to your module's private header files here.  You only need to
+// add includes for headers that are used in most of your module's source files though.
+
+#include "Analytics.h"
diff --git a/Engine/Plugins/Runtime/Analytics/Flurry/Source/FlurryEditor/Public/FlurryEditor.h b/Engine/Plugins/Runtime/Analytics/Flurry/Source/FlurryEditor/Public/FlurryEditor.h
new file mode 100644
--- /dev/null
+++ b/Engine/Plugins/Runtime/Analytics/Flurry/Source/FlurryEditor/Public/FlurryEditor.h
@@ -0,0 +1,12 @@
+// Copyright 1998-2015 Epic Games, Inc. All Rights Reserved.
+
+#pragma once
+
+class FFlurryEditorModule :
+	public IModuleInterface
+{
+private:
+	virtual void StartupModule() override;
+	virtual void ShutdownModule() override;
+};
+
diff --git a/Engine/Plugins/Runtime/Analytics/Flurry/Source/IOSFlurry/Source/IOSFlurry/IOSFlurry.Build.cs b/Engine/Plugins/Runtime/Analytics/Flurry/Source/IOSFlurry/Source/IOSFlurry/IOSFlurry.Build.cs
new file mode 100644
--- /dev/null
+++ b/Engine/Plugins/Runtime/Analytics/Flurry/Source/IOSFlurry/Source/IOSFlurry/IOSFlurry.Build.cs
@@ -0,0 +1,63 @@
+// Copyright 1998-2015 Epic Games, Inc. All Rights Reserved.
+
+namespace UnrealBuildTool.Rules
+{
+	public class IOSFlurry : ModuleRules
+	{
+        public IOSFlurry(TargetInfo Target)
+		{
+			BinariesSubFolder = "NotForLicensees";
+
+			PublicDependencyModuleNames.AddRange(
+				new string[]
+				{
+					"Core",
+					// ... add other public dependencies that you statically link with here ...
+				}
+				);
+
+			PrivateDependencyModuleNames.AddRange(
+				new string[]
+				{
+					"Analytics",
+					// ... add private dependencies that you statically link with here ...
+				}
+				);
+
+            PublicIncludePathModuleNames.Add("Analytics");
+
+			PublicFrameworks.AddRange(
+				new string[] {
+					"CoreTelephony",
+					"SystemConfiguration",
+					"UIKit",
+					"Foundation",
+					"CoreGraphics",
+					"MobileCoreServices",
+					"StoreKit",
+					"CFNetwork",
+					"CoreData",
+					"Security",
+					"CoreLocation"
+				});
+
+            bool bHasFlurrySDK =
+                (System.IO.Directory.Exists(System.IO.Path.Combine(UEBuildConfiguration.UEThirdPartySourceDirectory, "Flurry")) &&
+                  System.IO.Directory.Exists(System.IO.Path.Combine(UEBuildConfiguration.UEThirdPartySourceDirectory, "Flurry", "IOS"))) ||
+                (System.IO.Directory.Exists(System.IO.Path.Combine(UEBuildConfiguration.UEThirdPartySourceDirectory, "NotForLicensees")) &&
+                  System.IO.Directory.Exists(System.IO.Path.Combine(UEBuildConfiguration.UEThirdPartySourceDirectory, "NotForLicensees", "Flurry")) &&
+                  System.IO.Directory.Exists(System.IO.Path.Combine(UEBuildConfiguration.UEThirdPartySourceDirectory, "NotForLicensees", "Flurry", "IOS")));
+            if (bHasFlurrySDK)
+            {
+                PublicIncludePaths.Add(UEBuildConfiguration.UEThirdPartySourceDirectory + "NotForLicensees/Flurry/IOS/");
+                PublicAdditionalLibraries.Add(UEBuildConfiguration.UEThirdPartySourceDirectory + "NotForLicensees/Flurry/IOS/libFlurry_6.3.0.a");
+
+                Definitions.Add("WITH_FLURRY=1");
+            }
+            else
+            {
+                Definitions.Add("WITH_FLURRY=0");
+            }
+        }
+	}
+}
\ No newline at end of file
diff --git a/Engine/Plugins/Runtime/Analytics/Flurry/Source/IOSFlurry/Source/IOSFlurry/Private/IOSFlurry.cpp b/Engine/Plugins/Runtime/Analytics/Flurry/Source/IOSFlurry/Source/IOSFlurry/Private/IOSFlurry.cpp
new file mode 100644
--- /dev/null
+++ b/Engine/Plugins/Runtime/Analytics/Flurry/Source/IOSFlurry/Source/IOSFlurry/Private/IOSFlurry.cpp
@@ -0,0 +1,399 @@
+// Copyright 1998-2015 Epic Games, Inc. All Rights Reserved.
+
+#include "IOSFlurryPrivatePCH.h"
+
+DEFINE_LOG_CATEGORY_STATIC(LogAnalytics, Display, All);
+
+IMPLEMENT_MODULE( FAnalyticsIOSFlurry, IOSFlurry )
+
+TSharedPtr<IAnalyticsProvider> FAnalyticsProviderFlurry::Provider;
+
+void FAnalyticsIOSFlurry::StartupModule()
+{
+}
+
+void FAnalyticsIOSFlurry::ShutdownModule()
+{
+	FAnalyticsProviderFlurry::Destroy();
+}
+
+TSharedPtr<IAnalyticsProvider> FAnalyticsIOSFlurry::CreateAnalyticsProvider(const FAnalytics::FProviderConfigurationDelegate& GetConfigValue) const
+{
+	if (GetConfigValue.IsBound())
+	{
+		const FString Key = GetConfigValue.Execute(TEXT("FlurryApiKey"), true);
+		return FAnalyticsProviderFlurry::Create(Key);
+	}
+	else
+	{
+		UE_LOG(LogAnalytics, Warning, TEXT("IOSFlurry::CreateAnalyticsProvider called with an unbound config delegate"));
+	}
+	return nullptr;
+}
+
+// Provider
+
+FAnalyticsProviderFlurry::FAnalyticsProviderFlurry(const FString Key) :
+	ApiKey(Key)
+{
+#if WITH_FLURRY
+
+	// Support more config options here
+
+#else
+	UE_LOG(LogAnalytics, Warning, TEXT("WITH_FLURRY=0. Are you missing the SDK?"));
+#endif
+}
+
+FAnalyticsProviderFlurry::~FAnalyticsProviderFlurry()
+{
+	if (bHasSessionStarted)
+	{
+		EndSession();
+	}
+}
+
+bool FAnalyticsProviderFlurry::StartSession(const TArray<FAnalyticsEventAttribute>& Attributes)
+{
+#if WITH_FLURRY
+	if (!bHasSessionStarted)
+	{
+		NSString* ApiKeyStr = [NSString stringWithFString : ApiKey];
+		const int32 AttrCount = Attributes.Num();
+
+		[Flurry startSession : ApiKeyStr];
+		if (AttrCount > 0)
+		{
+			RecordEvent(TEXT("SessionAttributes"), Attributes);
+		}
+
+		UE_LOG(LogAnalytics, Display, TEXT("IOSFlurry::StartSession(%d attributes)"), AttrCount);
+		bHasSessionStarted = true;
+	}
+	return bHasSessionStarted;
+#else
+	UE_LOG(LogAnalytics, Warning, TEXT("WITH_FLURRY=0. Are you missing the SDK?"));
+	return false;
+#endif
+}
+
+void FAnalyticsProviderFlurry::EndSession()
+{
+#if WITH_FLURRY
+	// Flurry doesn't support ending a session
+	UE_LOG(LogAnalytics, Display, TEXT("IOSFlurry::EndSession - ignoring call"));
+#else
+	UE_LOG(LogAnalytics, Warning, TEXT("WITH_FLURRY=0. Are you missing the SDK?"));
+#endif
+}
+
+void FAnalyticsProviderFlurry::FlushEvents()
+{
+#if WITH_FLURRY
+	// Flurry doesn't support flushing a session
+	UE_LOG(LogAnalytics, Display, TEXT("IOSFlurry::FlushEvents - ignoring call"));
+#else
+	UE_LOG(LogAnalytics, Warning, TEXT("WITH_FLURRY=0. Are you missing the SDK?"));
+#endif
+}
+
+void FAnalyticsProviderFlurry::SetUserID(const FString& InUserID)
+{
+#if WITH_FLURRY
+	// Cache the string in case GetUserId is called
+	UserId = InUserID;
+	NSString* UserIdStr = [NSString stringWithFString : UserId];
+
+	[Flurry setUserID:UserIdStr];
+
+	UE_LOG(LogAnalytics, Display, TEXT("IOSFlurry::SetUserID(%s)"), *UserId);
+#else
+	UE_LOG(LogAnalytics, Warning, TEXT("WITH_FLURRY=0. Are you missing the SDK?"));
+#endif
+}
+
+FString FAnalyticsProviderFlurry::GetUserID() const
+{
+#if WITH_FLURRY
+	UE_LOG(LogAnalytics, Display, TEXT("IOSFlurry::GetUserID - returning cached id '%s'"), *UserId);
+
+	return UserId;
+#else
+	UE_LOG(LogAnalytics, Warning, TEXT("WITH_FLURRY=0. Are you missing the SDK?"));
+	return FString();
+#endif
+}
+
+void FAnalyticsProviderFlurry::SetGender(const FString& InGender)
+{
+#if WITH_FLURRY
+	NSString* ConvertedStr = [NSString stringWithFString : InGender];
+
+	[Flurry setGender:ConvertedStr];
+
+	UE_LOG(LogAnalytics, Display, TEXT("IOSFlurry::SetGender(%s)"), *InGender);
+#else
+	UE_LOG(LogAnalytics, Warning, TEXT("WITH_FLURRY=0. Are you missing the SDK?"));
+#endif
+}
+
+void FAnalyticsProviderFlurry::SetAge(const int32 InAge)
+{
+#if WITH_FLURRY
+	[Flurry setAge:InAge];
+
+	UE_LOG(LogAnalytics, Display, TEXT("IOSFlurry::SetAge(%d)"), InAge);
+#else
+	UE_LOG(LogAnalytics, Warning, TEXT("WITH_FLURRY=0. Are you missing the SDK?"));
+#endif
+}
+
+void FAnalyticsProviderFlurry::SetLocation(const FString& InLocation)
+{
+#if WITH_FLURRY
+	FString Lat, Long;
+	InLocation.Split(TEXT(","), &Lat, &Long);
+
+	double Latitude = FCString::Atod(*Lat);
+	double Longitude = FCString::Atod(*Long);
+
+	[Flurry setLatitude:Latitude longitude:Longitude horizontalAccuracy:0.0 verticalAccuracy:0.0];
+
+	UE_LOG(LogAnalytics, Display, TEXT("Parsed \"lat, long\" string in IOSFlurry::SetLocation(%s) as \"%f, %f\""), *InLocation, Latitude, Longitude);
+#else
+	UE_LOG(LogAnalytics, Warning, TEXT("WITH_FLURRY=0. Are you missing the SDK?"));
+#endif
+}
+
+FString FAnalyticsProviderFlurry::GetSessionID() const
+{
+#if WITH_FLURRY
+	NSString* Id = [Flurry getSessionID];
+	FString ConvertedId(Id);
+
+	UE_LOG(LogAnalytics, Display, TEXT("IOSFlurry::GetSessionID - returning the id as '%s'"), *ConvertedId);
+
+	return ConvertedId;
+#else
+	UE_LOG(LogAnalytics, Warning, TEXT("WITH_FLURRY=0. Are you missing the SDK?"));
+	return FString();
+#endif
+}
+
+bool FAnalyticsProviderFlurry::SetSessionID(const FString& InSessionID)
+{
+#if WITH_FLURRY
+	// Ignored
+	UE_LOG(LogAnalytics, Display, TEXT("IOSFlurry::SetSessionID - ignoring call"));
+	return false;
+#else
+	UE_LOG(LogAnalytics, Warning, TEXT("WITH_FLURRY=0. Are you missing the SDK?"));
+	return false;
+#endif
+}
+
+void FAnalyticsProviderFlurry::RecordEvent(const FString& EventName, const TArray<FAnalyticsEventAttribute>& Attributes)
+{
+#if WITH_FLURRY
+	if (EventName.Len() > 0)
+	{
+		NSString* ConvertedEventName = [NSString stringWithFString : EventName];
+		const int32 AttrCount = Attributes.Num();
+		if (AttrCount > 0)
+		{
+			// Convert the event attributes into a dictionary object
+			NSDictionary* AttributesDict = [NSMutableDictionary dictionaryWithCapacity:AttrCount];
+			for	(auto Attr : Attributes)
+			{
+				NSString* AttrName = [NSString stringWithFString : Attr.AttrName];
+				NSString* AttrValue = [NSString stringWithFString : Attr.AttrValue];
+				[AttributesDict setValue:AttrValue forKey:AttrName];
+			}
+			[Flurry logEvent:ConvertedEventName withParameters:AttributesDict];
+		}
+		else
+		{
+			// Log an event with no payload
+			[Flurry logEvent:ConvertedEventName];
+		}
+		UE_LOG(LogAnalytics, Display, TEXT("IOSFlurry::RecordEvent('%s', %d attributes)"), *EventName, AttrCount);
+	}
+#else
+	UE_LOG(LogAnalytics, Warning, TEXT("WITH_FLURRY=0. Are you missing the SDK?"));
+#endif
+}
+
+void FAnalyticsProviderFlurry::RecordItemPurchase(const FString& ItemId, const FString& Currency, int PerItemCost, int ItemQuantity)
+{
+#if WITH_FLURRY
+	NSString* EventName = @"Item Purchase";
+	// Build the dictionary
+	NSDictionary* AttributesDict = [NSMutableDictionary dictionaryWithCapacity:4];
+	[AttributesDict setValue:[NSString stringWithFString:ItemId] forKey:@"ItemId"];
+	[AttributesDict setValue:[NSString stringWithFString:Currency] forKey:@"Currency"];
+	[AttributesDict setValue:[NSString stringWithFormat:@"%d", PerItemCost] forKey:@"PerItemCost"];
+	[AttributesDict setValue:[NSString stringWithFormat:@"%d", ItemQuantity] forKey:@"ItemQuantity"];
+	// Send the event
+	[Flurry logEvent:EventName withParameters:AttributesDict];
+
+	UE_LOG(LogAnalytics, Display, TEXT("IOSFlurry::RecordItemPurchase('%s', '%s', %d, %d)"), *ItemId, *Currency, PerItemCost, ItemQuantity);
+#else
+	UE_LOG(LogAnalytics, Warning, TEXT("WITH_FLURRY=0. Are you missing the SDK?"));
+#endif
+}
+
+void FAnalyticsProviderFlurry::RecordCurrencyPurchase(const FString& GameCurrencyType, int GameCurrencyAmount, const FString& RealCurrencyType, float RealMoneyCost, const FString& PaymentProvider)
+{
+#if WITH_FLURRY
+	NSString* EventName = @"Currency Purchase";
+	// Build the dictionary
+	NSDictionary* AttributesDict = [NSMutableDictionary dictionaryWithCapacity:5];
+	[AttributesDict setValue:[NSString stringWithFString:GameCurrencyType] forKey:@"GameCurrencyType"];
+	[AttributesDict setValue:[NSString stringWithFormat:@"%d", GameCurrencyAmount] forKey:@"GameCurrencyAmount"];
+	[AttributesDict setValue:[NSString stringWithFString:RealCurrencyType] forKey:@"RealCurrencyType"];
+	[AttributesDict setValue:[NSString stringWithFormat:@"%.02f", RealMoneyCost] forKey:@"RealMoneyCost"];
+	[AttributesDict setValue:[NSString stringWithFString:PaymentProvider] forKey:@"PaymentProvider"];
+	// Send the event
+	[Flurry logEvent:EventName withParameters:AttributesDict];
+
+	UE_LOG(LogAnalytics, Display, TEXT("IOSFlurry::RecordCurrencyPurchase('%s', %d, '%s', %.02f, %s)"), *GameCurrencyType, GameCurrencyAmount, *RealCurrencyType, RealMoneyCost, *PaymentProvider);
+#else
+	UE_LOG(LogAnalytics, Warning, TEXT("WITH_FLURRY=0. Are you missing the SDK?"));
+#endif
+}
+
+void FAnalyticsProviderFlurry::RecordCurrencyGiven(const FString& GameCurrencyType, int GameCurrencyAmount)
+{
+#if WITH_FLURRY
+	NSString* EventName = @"Currency Given";
+	// Build the dictionary
+	NSDictionary* AttributesDict = [NSMutableDictionary dictionaryWithCapacity:2];
+	[AttributesDict setValue:[NSString stringWithFString:GameCurrencyType] forKey:@"GameCurrencyType"];
+	[AttributesDict setValue:[NSString stringWithFormat:@"%d", GameCurrencyAmount] forKey:@"GameCurrencyAmount"];
+	// Send the event
+	[Flurry logEvent:EventName withParameters:AttributesDict];
+
+	UE_LOG(LogAnalytics, Display, TEXT("IOSFlurry::RecordCurrencyGiven('%s', %d)"), *GameCurrencyType, GameCurrencyAmount);
+#else
+	UE_LOG(LogAnalytics, Warning, TEXT("WITH_FLURRY=0. Are you missing the SDK?"));
+#endif
+}
+
+void FAnalyticsProviderFlurry::RecordItemPurchase(const FString& ItemId, int ItemQuantity, const TArray<FAnalyticsEventAttribute>& EventAttrs)
+{
+#if WITH_FLURRY
+	NSString* EventName = @"Item Purchase";
+	// Build the dictionary
+	int32 DictSize = EventAttrs.Num() + 2;
+	NSDictionary* AttributesDict = [NSMutableDictionary dictionaryWithCapacity:DictSize];
+	[AttributesDict setValue:[NSString stringWithFString:ItemId] forKey:@"ItemId"];
+	[AttributesDict setValue:[NSString stringWithFormat:@"%d", ItemQuantity] forKey:@"Quantity"];
+	for	(auto Attr : EventAttrs)
+	{
+		NSString* AttrName = [NSString stringWithFString : Attr.AttrName];
+		NSString* AttrValue = [NSString stringWithFString : Attr.AttrValue];
+		[AttributesDict setValue:AttrValue forKey:AttrName];
+	}
+	// Send the event
+	[Flurry logEvent:EventName withParameters:AttributesDict];
+
+	UE_LOG(LogAnalytics, Display, TEXT("IOSFlurry::RecordItemPurchase('%s', %d, %d)"), *ItemId, ItemQuantity, EventAttrs.Num());
+#else
+	UE_LOG(LogAnalytics, Warning, TEXT("WITH_FLURRY=0. Are you missing the SDK?"));
+#endif
+}
+
+void FAnalyticsProviderFlurry::RecordCurrencyPurchase(const FString& GameCurrencyType, int GameCurrencyAmount, const TArray<FAnalyticsEventAttribute>& EventAttrs)
+{
+#if WITH_FLURRY
+	NSString* EventName = @"Currency Purchase";
+	// Build the dictionary
+	int32 DictSize = EventAttrs.Num() + 2;
+	NSDictionary* AttributesDict = [NSMutableDictionary dictionaryWithCapacity:DictSize];
+	[AttributesDict setValue:[NSString stringWithFString:GameCurrencyType] forKey:@"GameCurrencyType"];
+	[AttributesDict setValue:[NSString stringWithFormat:@"%d", GameCurrencyAmount] forKey:@"GameCurrencyAmount"];
+	for	(auto Attr : EventAttrs)
+	{
+		NSString* AttrName = [NSString stringWithFString : Attr.AttrName];
+		NSString* AttrValue = [NSString stringWithFString : Attr.AttrValue];
+		[AttributesDict setValue:AttrValue forKey:AttrName];
+	}
+	// Send the event
+	[Flurry logEvent:EventName withParameters:AttributesDict];
+
+	UE_LOG(LogAnalytics, Display, TEXT("IOSFlurry::RecordCurrencyPurchase('%s', %d, %d)"), *GameCurrencyType, GameCurrencyAmount, EventAttrs.Num());
+#else
+	UE_LOG(LogAnalytics, Warning, TEXT("WITH_FLURRY=0. Are you missing the SDK?"));
+#endif
+}
+
+void FAnalyticsProviderFlurry::RecordCurrencyGiven(const FString& GameCurrencyType, int GameCurrencyAmount, const TArray<FAnalyticsEventAttribute>& EventAttrs)
+{
+#if WITH_FLURRY
+	NSString* EventName = @"Currency Given";
+	// Build the dictionary
+	int32 DictSize = EventAttrs.Num() + 2;
+	NSDictionary* AttributesDict = [NSMutableDictionary dictionaryWithCapacity:DictSize];
+	[AttributesDict setValue:[NSString stringWithFString:GameCurrencyType] forKey:@"GameCurrencyType"];
+	[AttributesDict setValue:[NSString stringWithFormat:@"%d", GameCurrencyAmount] forKey:@"GameCurrencyAmount"];
+	for	(auto Attr : EventAttrs)
+	{
+		NSString* AttrName = [NSString stringWithFString : Attr.AttrName];
+		NSString* AttrValue = [NSString stringWithFString : Attr.AttrValue];
+		[AttributesDict setValue:AttrValue forKey:AttrName];
+	}
+	// Send the event
+	[Flurry logEvent:EventName withParameters:AttributesDict];
+
+	UE_LOG(LogAnalytics, Display, TEXT("IOSFlurry::RecordCurrencyGiven('%s', %d, %d)"), *GameCurrencyType, GameCurrencyAmount, EventAttrs.Num());
+#else
+	UE_LOG(LogAnalytics, Warning, TEXT("WITH_FLURRY=0. Are you missing the SDK?"));
+#endif
+}
+
+void FAnalyticsProviderFlurry::RecordError(const FString& Error, const TArray<FAnalyticsEventAttribute>& EventAttrs)
+{
+#if WITH_FLURRY
+	NSString* EventName = @"Error";
+	// Build the dictionary
+	int32 DictSize = EventAttrs.Num() + 1;
+	NSDictionary* AttributesDict = [NSMutableDictionary dictionaryWithCapacity:DictSize];
+	[AttributesDict setValue:[NSString stringWithFString:Error] forKey:@"Error"];
+	for	(auto Attr : EventAttrs)
+	{
+		NSString* AttrName = [NSString stringWithFString : Attr.AttrName];
+		NSString* AttrValue = [NSString stringWithFString : Attr.AttrValue];
+		[AttributesDict setValue:AttrValue forKey:AttrName];
+	}
+	// Send the event
+	[Flurry logEvent:EventName withParameters:AttributesDict];
+
+	UE_LOG(LogAnalytics, Display, TEXT("IOSFlurry::RecordError('%s', %d)"), *Error, EventAttrs.Num());
+#else
+	UE_LOG(LogAnalytics, Warning, TEXT("WITH_FLURRY=0. Are you missing the SDK?"));
+#endif
+}
+
+void FAnalyticsProviderFlurry::RecordProgress(const FString& ProgressType, const FString& ProgressHierarchy, const TArray<FAnalyticsEventAttribute>& EventAttrs)
+{
+#if WITH_FLURRY
+	NSString* EventName = @"Progress";
+	// Build the dictionary
+	int32 DictSize = EventAttrs.Num() + 2;
+	NSDictionary* AttributesDict = [NSMutableDictionary dictionaryWithCapacity:DictSize];
+	[AttributesDict setValue:[NSString stringWithFString:ProgressType] forKey:@"ProgressType"];
+	[AttributesDict setValue:[NSString stringWithFString:ProgressHierarchy] forKey:@"ProgressHierarchy"];
+	for	(auto Attr : EventAttrs)
+	{
+		NSString* AttrName = [NSString stringWithFString : Attr.AttrName];
+		NSString* AttrValue = [NSString stringWithFString : Attr.AttrValue];
+		[AttributesDict setValue:AttrValue forKey:AttrName];
+	}
+	// Send the event
+	[Flurry logEvent:EventName withParameters:AttributesDict];
+
+	UE_LOG(LogAnalytics, Display, TEXT("IOSFlurry::RecordProgress('%s', %s, %d)"), *ProgressType, *ProgressHierarchy, EventAttrs.Num());
+#else
+	UE_LOG(LogAnalytics, Warning, TEXT("WITH_FLURRY=0. Are you missing the SDK?"));
+#endif
+}
diff --git a/Engine/Plugins/Runtime/Analytics/Flurry/Source/IOSFlurry/Source/IOSFlurry/Private/IOSFlurryPrivatePCH.h b/Engine/Plugins/Runtime/Analytics/Flurry/Source/IOSFlurry/Source/IOSFlurry/Private/IOSFlurryPrivatePCH.h
new file mode 100644
--- /dev/null
+++ b/Engine/Plugins/Runtime/Analytics/Flurry/Source/IOSFlurry/Source/IOSFlurry/Private/IOSFlurryPrivatePCH.h
@@ -0,0 +1,15 @@
+// Copyright 1998-2015 Epic Games, Inc. All Rights Reserved.
+
+#include "Core.h"
+#include "IAnalyticsProviderModule.h"
+#include "IAnalyticsProvider.h"
+
+#include "IOSFlurry.h"
+#include "IOSFlurryProvider.h"
+
+#if WITH_FLURRY
+#import "Flurry.h"
+#endif
+
+// You should place include statements to your module's private header files here.  You only need to
+// add includes for headers that are used in most of your module's source files though.
diff --git a/Engine/Plugins/Runtime/Analytics/Flurry/Source/IOSFlurry/Source/IOSFlurry/Private/IOSFlurryProvider.h b/Engine/Plugins/Runtime/Analytics/Flurry/Source/IOSFlurry/Source/IOSFlurry/Private/IOSFlurryProvider.h
new file mode 100644
--- /dev/null
+++ b/Engine/Plugins/Runtime/Analytics/Flurry/Source/IOSFlurry/Source/IOSFlurry/Private/IOSFlurryProvider.h
@@ -0,0 +1,60 @@
+// Copyright 1998-2015 Epic Games, Inc. All Rights Reserved.
+
+#pragma once
+
+class FAnalyticsProviderFlurry :
+	public IAnalyticsProvider
+{
+	/** The API key given to you by Flurry */
+	FString ApiKey;
+	/** Tracks whether we need to start the session or restart it */
+	bool bHasSessionStarted;
+	/** Cached user id */
+	FString UserId;
+
+	/** Singleton for analytics */
+	static TSharedPtr<IAnalyticsProvider> Provider;
+	FAnalyticsProviderFlurry(const FString Key);
+
+public:
+	static TSharedPtr<IAnalyticsProvider> Create(const FString Key)
+	{
+		if (!Provider.IsValid())
+		{
+			Provider = TSharedPtr<IAnalyticsProvider>(new FAnalyticsProviderFlurry(Key));
+		}
+		return Provider;
+	}
+	static void Destroy()
+	{
+		Provider.Reset();
+	}
+
+	virtual ~FAnalyticsProviderFlurry();
+
+	virtual bool StartSession(const TArray<FAnalyticsEventAttribute>& Attributes) override;
+	virtual void EndSession() override;
+	virtual void FlushEvents() override;
+
+	virtual void SetUserID(const FString& InUserID) override;
+	virtual FString GetUserID() const override;
+
+	virtual FString GetSessionID() const override;
+	virtual bool SetSessionID(const FString& InSessionID) override;
+
+	virtual void RecordEvent(const FString& EventName, const TArray<FAnalyticsEventAttribute>& Attributes) override;
+
+	virtual void RecordItemPurchase(const FString& ItemId, const FString& Currency, int PerItemCost, int ItemQuantity) override;
+	virtual void RecordCurrencyPurchase(const FString& GameCurrencyType, int GameCurrencyAmount, const FString& RealCurrencyType, float RealMoneyCost, const FString& PaymentProvider) override;
+	virtual void RecordCurrencyGiven(const FString& GameCurrencyType, int GameCurrencyAmount) override;
+
+	virtual void SetGender(const FString& InGender) override;
+	virtual void SetLocation(const FString& InLocation) override;
+	virtual void SetAge(const int32 InAge) override;
+
+	virtual void RecordItemPurchase(const FString& ItemId, int ItemQuantity, const TArray<FAnalyticsEventAttribute>& EventAttrs) override;
+	virtual void RecordCurrencyPurchase(const FString& GameCurrencyType, int GameCurrencyAmount, const TArray<FAnalyticsEventAttribute>& EventAttrs) override;
+	virtual void RecordCurrencyGiven(const FString& GameCurrencyType, int GameCurrencyAmount, const TArray<FAnalyticsEventAttribute>& EventAttrs) override;
+	virtual void RecordError(const FString& Error, const TArray<FAnalyticsEventAttribute>& EventAttrs) override;
+	virtual void RecordProgress(const FString& ProgressType, const FString& ProgressHierarchy, const TArray<FAnalyticsEventAttribute>& EventAttrs) override;
+};
diff --git a/Engine/Plugins/Runtime/Analytics/Flurry/Source/IOSFlurry/Source/IOSFlurry/Public/IOSFlurry.h b/Engine/Plugins/Runtime/Analytics/Flurry/Source/IOSFlurry/Source/IOSFlurry/Public/IOSFlurry.h
new file mode 100644
--- /dev/null
+++ b/Engine/Plugins/Runtime/Analytics/Flurry/Source/IOSFlurry/Source/IOSFlurry/Public/IOSFlurry.h
@@ -0,0 +1,49 @@
+// Copyright 1998-2015 Epic Games, Inc. All Rights Reserved.
+
+#pragma once
+
+#include "Runtime/Analytics/Analytics/Public/Interfaces/IAnalyticsProviderModule.h"
+#include "Core.h"
+
+class IAnalyticsProvider;
+
+/**
+ * The public interface to this module
+ */
+class FAnalyticsIOSFlurry :
+	public IAnalyticsProviderModule
+{
+	/** Singleton for analytics */
+	TSharedPtr<IAnalyticsProvider> Provider;
+
+	//--------------------------------------------------------------------------
+	// Module functionality
+	//--------------------------------------------------------------------------
+public:
+	/**
+	 * Singleton-like access to this module's interface.  This is just for convenience!
+	 * Beware of calling this during the shutdown phase, though.  Your module might have been unloaded already.
+	 *
+	 * @return Returns singleton instance, loading the module on demand if needed
+	 */
+	static inline FAnalyticsIOSFlurry& Get()
+	{
+		return FModuleManager::LoadModuleChecked< FAnalyticsIOSFlurry >( "IOSFlurry" );
+	}
+
+	//--------------------------------------------------------------------------
+	// provider factory functions
+	//--------------------------------------------------------------------------
+public:
+	/**
+	 * IAnalyticsProviderModule interface.
+	 * Creates the analytics provider given a configuration delegate.
+	 * The keys required exactly match the field names in the Config object. 
+	 */
+	virtual TSharedPtr<IAnalyticsProvider> CreateAnalyticsProvider(const FAnalytics::FProviderConfigurationDelegate& GetConfigValue) const override;
+	
+private:
+	virtual void StartupModule() override;
+	virtual void ShutdownModule() override;
+};
+
diff --git a/Engine/Source/Runtime/Analytics/Analytics/Private/Analytics.cpp b/Engine/Source/Runtime/Analytics/Analytics/Private/Analytics.cpp
--- a/Engine/Source/Runtime/Analytics/Analytics/Private/Analytics.cpp
+++ b/Engine/Source/Runtime/Analytics/Analytics/Private/Analytics.cpp
@@ -69,6 +69,12 @@ FString FAnalytics::GetConfigValueFromIni( const FString& IniName, const FString
 	return Result;
 }
 
+void FAnalytics::WriteConfigValueToIni(const FString& IniName, const FString& SectionName, const FString& KeyName, const FString& Value)
+{
+	GConfig->SetString(*SectionName, *KeyName, *Value, *IniName);
+}
+
+
 void FAnalytics::StartupModule()
 {
 }
diff --git a/Engine/Source/Runtime/Analytics/Analytics/Public/Analytics.h b/Engine/Source/Runtime/Analytics/Analytics/Public/Analytics.h
--- a/Engine/Source/Runtime/Analytics/Analytics/Public/Analytics.h
+++ b/Engine/Source/Runtime/Analytics/Analytics/Public/Analytics.h
@@ -217,6 +217,10 @@ public:
 	 */
 	virtual FString GetConfigValueFromIni(const FString& IniName, const FString& SectionName, const FString& KeyName, bool bIsRequired);
 
+	/**
+	 * Helper for writing configuration values from to an INI file (which will be a common scenario). 
+	 */
+	virtual void WriteConfigValueToIni(const FString& IniName, const FString& SectionName, const FString& KeyName, const FString& Value);
 private:
 	virtual void StartupModule() override;
 	virtual void ShutdownModule() override;
diff --git a/Engine/Source/Runtime/Analytics/AnalyticsVisualEditing/AnalyticsVisualEditing.Build.cs b/Engine/Source/Runtime/Analytics/AnalyticsVisualEditing/AnalyticsVisualEditing.Build.cs
new file mode 100644
--- /dev/null
+++ b/Engine/Source/Runtime/Analytics/AnalyticsVisualEditing/AnalyticsVisualEditing.Build.cs
@@ -0,0 +1,20 @@
+// Copyright 1998-2015 Epic Games, Inc. All Rights Reserved.
+
+namespace UnrealBuildTool.Rules
+{
+	public class AnalyticsVisualEditing : ModuleRules
+	{
+        public AnalyticsVisualEditing(TargetInfo Target)
+		{
+            PrivateDependencyModuleNames.AddRange(
+                new string[]
+                {
+					"Analytics",
+    				"Core",
+	    			"CoreUObject",
+                    "Engine"
+				}
+			);
+		}
+	}
+}
diff --git a/Engine/Source/Runtime/Analytics/AnalyticsVisualEditing/Classes/AnalyticsSettings.h b/Engine/Source/Runtime/Analytics/AnalyticsVisualEditing/Classes/AnalyticsSettings.h
new file mode 100644
--- /dev/null
+++ b/Engine/Source/Runtime/Analytics/AnalyticsVisualEditing/Classes/AnalyticsSettings.h
@@ -0,0 +1,110 @@
+// Copyright 1998-2015 Epic Games, Inc. All Rights Reserved.
+
+#pragma once
+
+#include "AnalyticsSettings.generated.h"
+
+UCLASS(Abstract)
+class ANALYTICSVISUALEDITING_API UAnalyticsSettingsBase
+	: public UDeveloperSettings
+{
+	GENERATED_UCLASS_BODY()
+
+	UPROPERTY()
+	FText SettingsDisplayName;
+
+	UPROPERTY()
+	FText SettingsTooltip;
+
+	/** Helpers so that subclasses don't need to hardcode the strings in case they change */
+	FORCEINLINE FString GetReleaseIniSection() const { return TEXT("Analytics"); }
+	FORCEINLINE FString GetDebugIniSection() const { return TEXT("AnalyticsDebug"); }
+	FORCEINLINE FString GetTestIniSection() const { return TEXT("AnalyticsTest"); }
+	FORCEINLINE FString GetDevelopmentIniSection() const { return TEXT("AnalyticsDevelopment"); }
+	FORCEINLINE FString GetIniName() const { return FString::Printf(TEXT("%sDefaultEngine.ini"), *FPaths::SourceConfigDir()); }
+
+private:
+/**
+ * Because the analytics providers can be used outside of apps that include the UObject framework
+ * we have to manually hook up the INI loading/saving which is done by these methods
+ */
+#if WITH_EDITOR
+	virtual void PostInitProperties() override
+	{
+		Super::PostInitProperties();
+
+		ReadConfigSettings();
+	}
+
+	virtual void PostEditChangeProperty(FPropertyChangedEvent& PropertyChangedEvent) override
+	{
+		Super::PostEditChangeProperty(PropertyChangedEvent);
+
+		WriteConfigSettings();
+
+		GConfig->Flush(false, GetIniName());
+	}
+#endif
+
+// UDeveloperSettings interface
+	/** Gets the category for the settings, some high level grouping like, Editor, Engine, Game...etc. */
+	virtual FName GetCategoryName() const override;
+
+// UAnalyticsSettingsBase interface  - derived classes should implement
+protected:
+	/**
+	 * Provides a mechanism to read the section based information into this UObject's properties
+	 */
+	virtual void ReadConfigSettings() {}
+	/**
+	 * Provides a mechanism to save this object's properties to the section based ini values
+	 */
+	virtual void WriteConfigSettings() {}
+
+#if WITH_EDITOR
+	/** Gets the section text for this  */
+	virtual FText GetSectionText() const
+	{
+		return SettingsDisplayName;
+	}
+	/** Gets the description for the section, uses the classes ToolTip by default. */
+	virtual FText GetSectionDescription() const
+	{
+		return SettingsTooltip;
+	}
+#endif
+};
+
+UCLASS()
+class UAnalyticsSettings
+	: public UAnalyticsSettingsBase
+{
+	GENERATED_UCLASS_BODY()
+
+	/** The name of the plugin containing your desired analytics provider */
+	UPROPERTY(EditAnywhere, DisplayName="Development", Category=Providers, meta=(ConfigRestartRequired=true))
+	FString DevelopmentProviderName;
+
+	/** The name of the plugin containing your desired analytics provider */
+	UPROPERTY(EditAnywhere, DisplayName="Debug", Category=Providers, meta=(ConfigRestartRequired=true))
+	FString DebugProviderName;
+
+	/** The name of the plugin containing your desired analytics provider */
+	UPROPERTY(EditAnywhere, DisplayName="Test", Category=Providers, meta=(ConfigRestartRequired=true))
+	FString TestProviderName;
+
+	/** The name of the plugin containing your desired analytics provider */
+	UPROPERTY(EditAnywhere, DisplayName="Release", Category=Providers, meta=(ConfigRestartRequired=true))
+	FString ReleaseProviderName;
+
+// UAnalyticsSettingsBase interface
+protected:
+	/**
+	 * Provides a mechanism to read the section based information into this UObject's properties
+	 */
+	virtual void ReadConfigSettings();
+	/**
+	 * Provides a mechanism to save this object's properties to the section based ini values
+	 */
+	virtual void WriteConfigSettings();
+};
diff --git a/Engine/Source/Runtime/Analytics/AnalyticsVisualEditing/Private/AnalyticsVisualEditing.cpp b/Engine/Source/Runtime/Analytics/AnalyticsVisualEditing/Private/AnalyticsVisualEditing.cpp
new file mode 100644
--- /dev/null
+++ b/Engine/Source/Runtime/Analytics/AnalyticsVisualEditing/Private/AnalyticsVisualEditing.cpp
@@ -0,0 +1,57 @@
+// Copyright 1998-2015 Epic Games, Inc. All Rights Reserved.
+
+#include "Analytics.h"
+#include "Engine/DeveloperSettings.h"
+#include "AnalyticsVisualEditing.h"
+#include "AnalyticsSettings.h"
+
+IMPLEMENT_MODULE( FAnalyticsVisualEditingModule, AnalyticsVisualEditing );
+
+static const FName AnalyticsName(TEXT("Analytics"));
+
+#define LOCTEXT_NAMESPACE "Analytics"
+
+UAnalyticsSettingsBase::UAnalyticsSettingsBase(const FObjectInitializer& ObjectInitializer)
+	: Super(ObjectInitializer)
+{
+}
+
+UAnalyticsSettings::UAnalyticsSettings(const FObjectInitializer& ObjectInitializer)
+	: Super(ObjectInitializer)
+{
+	SettingsDisplayName = LOCTEXT("SettingsDisplayName", "Providers");
+	SettingsTooltip = LOCTEXT("SettingsTooltip", "Configures which analytics provider to use per build type");
+}
+
+FName UAnalyticsSettingsBase::GetCategoryName() const
+{
+	return AnalyticsName;
+}
+
+void UAnalyticsSettings::ReadConfigSettings()
+{
+	ReleaseProviderName = FAnalytics::Get().GetConfigValueFromIni(GetIniName(), GetReleaseIniSection(), TEXT("ProviderModuleName"), true);
+	DevelopmentProviderName = FAnalytics::Get().GetConfigValueFromIni(GetIniName(), GetDevelopmentIniSection(), TEXT("ProviderModuleName"), true);
+	if (DevelopmentProviderName.Len() == 0)
+	{
+		DevelopmentProviderName = ReleaseProviderName;
+	}
+	TestProviderName = FAnalytics::Get().GetConfigValueFromIni(GetIniName(), GetTestIniSection(), TEXT("ProviderModuleName"), true);
+	if (TestProviderName.Len() == 0)
+	{
+		TestProviderName = ReleaseProviderName;
+	}
+	DebugProviderName = FAnalytics::Get().GetConfigValueFromIni(GetIniName(), GetDebugIniSection(), TEXT("ProviderModuleName"), true);
+	if (DebugProviderName.Len() == 0)
+	{
+		DebugProviderName = ReleaseProviderName;
+	}
+}
+
+void UAnalyticsSettings::WriteConfigSettings()
+{
+	FAnalytics::Get().WriteConfigValueToIni(GetIniName(), GetReleaseIniSection(), TEXT("ProviderModuleName"), ReleaseProviderName);
+	FAnalytics::Get().WriteConfigValueToIni(GetIniName(), GetDevelopmentIniSection(), TEXT("ProviderModuleName"), DevelopmentProviderName);
+	FAnalytics::Get().WriteConfigValueToIni(GetIniName(), GetTestIniSection(), TEXT("ProviderModuleName"), TestProviderName);
+	FAnalytics::Get().WriteConfigValueToIni(GetIniName(), GetDebugIniSection(), TEXT("ProviderModuleName"), DebugProviderName);
+}
diff --git a/Engine/Source/Runtime/Analytics/AnalyticsVisualEditing/Public/AnalyticsVisualEditing.h b/Engine/Source/Runtime/Analytics/AnalyticsVisualEditing/Public/AnalyticsVisualEditing.h
new file mode 100644
--- /dev/null
+++ b/Engine/Source/Runtime/Analytics/AnalyticsVisualEditing/Public/AnalyticsVisualEditing.h
@@ -0,0 +1,14 @@
+// Copyright 1998-2015 Epic Games, Inc. All Rights Reserved.
+
+#pragma once
+
+/**
+ * This module holds the base UObject class that allows visual editing of configuration settings in the Unreal Editor.
+ * This is needed because the settings need to be usable outside of the UObject framework so these objects wrap
+ * those settings for editor display purposes and write the non-UObject config property information out
+ */
+class FAnalyticsVisualEditingModule :
+	public IModuleInterface
+{
+};
+
