Update the Flurry build file to the latest version of the library

diff --git a/Engine/Plugins/Runtime/Analytics/Flurry/Source/IOSFlurry/Source/IOSFlurry/IOSFlurry.Build.cs b/Engine/Plugins/Runtime/Analytics/Flurry/Source/IOSFlurry/Source/IOSFlurry/IOSFlurry.Build.cs
--- a/Engine/Plugins/Runtime/Analytics/Flurry/Source/IOSFlurry/Source/IOSFlurry/IOSFlurry.Build.cs
+++ b/Engine/Plugins/Runtime/Analytics/Flurry/Source/IOSFlurry/Source/IOSFlurry/IOSFlurry.Build.cs
@@ -50,7 +50,7 @@ namespace UnrealBuildTool.Rules
             if (bHasFlurrySDK)
             {
                 PublicIncludePaths.Add(UEBuildConfiguration.UEThirdPartySourceDirectory + "NotForLicensees/Flurry/IOS/");
-                PublicAdditionalLibraries.Add(UEBuildConfiguration.UEThirdPartySourceDirectory + "NotForLicensees/Flurry/IOS/libFlurry_6.3.0.a");
+                PublicAdditionalLibraries.Add(UEBuildConfiguration.UEThirdPartySourceDirectory + "NotForLicensees/Flurry/IOS/libFlurry_6.7.0.a");
 
                 Definitions.Add("WITH_FLURRY=1");
             }
