Fix for IOSFlurry AnalyticsProvider

diff --git a/Engine/Plugins/Runtime/Analytics/Flurry/Source/IOSFlurry/Source/IOSFlurry/Private/IOSFlurry.cpp b/Engine/Plugins/Runtime/Analytics/Flurry/Source/IOSFlurry/Source/IOSFlurry/Private/IOSFlurry.cpp
--- a/Engine/Plugins/Runtime/Analytics/Flurry/Source/IOSFlurry/Source/IOSFlurry/Private/IOSFlurry.cpp
+++ b/Engine/Plugins/Runtime/Analytics/Flurry/Source/IOSFlurry/Source/IOSFlurry/Private/IOSFlurry.cpp
@@ -202,13 +202,17 @@ void FAnalyticsProviderFlurry::RecordEvent(const FString& EventName, const TArra
 #if WITH_FLURRY
 	if (EventName.Len() > 0)
 	{
+		TArray<FAnalyticsEventAttribute> EventAttributes;
+		EventAttributes.Append(Attributes);
+		EventAttributes.Append(DefaultEventAttributes);
+
 		NSString* ConvertedEventName = [NSString stringWithFString : EventName];
-		const int32 AttrCount = Attributes.Num();
+		const int32 AttrCount = EventAttributes.Num();
 		if (AttrCount > 0)
 		{
 			// Convert the event attributes into a dictionary object
 			NSDictionary* AttributesDict = [NSMutableDictionary dictionaryWithCapacity:AttrCount];
-			for	(auto Attr : Attributes)
+			for	(auto Attr : EventAttributes)
 			{
 				NSString* AttrName = [NSString stringWithFString : Attr.GetName()];
 				NSString* AttrValue = [NSString stringWithFString : Attr.GetValue()];
@@ -402,3 +406,24 @@ void FAnalyticsProviderFlurry::RecordProgress(const FString& ProgressType, const
 	UE_LOG(LogAnalytics, Warning, TEXT("WITH_FLURRY=0. Are you missing the SDK?"));
 #endif
 }
+
+
+void FAnalyticsProviderFlurry::SetDefaultEventAttributes(TArray<FAnalyticsEventAttribute>&& Attributes)
+{
+	DefaultEventAttributes = Attributes;
+}
+
+TArray<FAnalyticsEventAttribute> FAnalyticsProviderFlurry::GetDefaultEventAttributesSafe() const
+{
+	return DefaultEventAttributes;
+}
+
+int32 FAnalyticsProviderFlurry::GetDefaultEventAttributeCount() const
+{
+	return DefaultEventAttributes.Num();
+}
+
+FAnalyticsEventAttribute FAnalyticsProviderFlurry::GetDefaultEventAttribute(int AttributeIndex) const
+{
+	return DefaultEventAttributes[AttributeIndex];
+}
diff --git a/Engine/Plugins/Runtime/Analytics/Flurry/Source/IOSFlurry/Source/IOSFlurry/Private/IOSFlurryProvider.h b/Engine/Plugins/Runtime/Analytics/Flurry/Source/IOSFlurry/Source/IOSFlurry/Private/IOSFlurryProvider.h
--- a/Engine/Plugins/Runtime/Analytics/Flurry/Source/IOSFlurry/Source/IOSFlurry/Private/IOSFlurryProvider.h
+++ b/Engine/Plugins/Runtime/Analytics/Flurry/Source/IOSFlurry/Source/IOSFlurry/Private/IOSFlurryProvider.h
@@ -59,4 +59,12 @@ public:
 	virtual void RecordCurrencyGiven(const FString& GameCurrencyType, int GameCurrencyAmount, const TArray<FAnalyticsEventAttribute>& EventAttrs) override;
 	virtual void RecordError(const FString& Error, const TArray<FAnalyticsEventAttribute>& EventAttrs) override;
 	virtual void RecordProgress(const FString& ProgressType, const FString& ProgressHierarchy, const TArray<FAnalyticsEventAttribute>& EventAttrs) override;
+
+	virtual void SetDefaultEventAttributes(TArray<FAnalyticsEventAttribute>&& Attributes)  override;
+	virtual TArray<FAnalyticsEventAttribute> GetDefaultEventAttributesSafe() const override;
+	virtual int32 GetDefaultEventAttributeCount() const  override;
+	virtual FAnalyticsEventAttribute GetDefaultEventAttribute(int AttributeIndex) const  override;
+
+private:
+	TArray<FAnalyticsEventAttribute>		DefaultEventAttributes;
 };
